Fixed non-unity USD compile issues.

diff --git a/Engine/Plugins/Runtime/USDCore/Source/UnrealUSDWrapper/Public/USDIncludesEnd.h b/Engine/Plugins/Runtime/USDCore/Source/UnrealUSDWrapper/Public/USDIncludesEnd.h
--- a/Engine/Plugins/Runtime/USDCore/Source/UnrealUSDWrapper/Public/USDIncludesEnd.h
+++ b/Engine/Plugins/Runtime/USDCore/Source/UnrealUSDWrapper/Public/USDIncludesEnd.h
@@ -6,7 +6,7 @@
 
 #if USE_USD_SDK
 
-#if !defined(UE_DIRECT_HEADER_COMPILE)
+#ifdef USD_INCLUDES_START
 
 #pragma warning(pop)
 #pragma pop_macro("check")
@@ -27,7 +27,11 @@
 
 THIRD_PARTY_INCLUDES_END
 
-#endif // UE_DIRECT_HEADER_COMPILE
+#undef USD_INCLUDES_START
+
+#else
+	#error Mismatched USDIncludesEnd.h detected.
+#endif // USD_INCLUDES_START
 
 #endif // #if USE_USD_SDK
 
diff --git a/Engine/Plugins/Runtime/USDCore/Source/UnrealUSDWrapper/Public/USDIncludesStart.h b/Engine/Plugins/Runtime/USDCore/Source/UnrealUSDWrapper/Public/USDIncludesStart.h
--- a/Engine/Plugins/Runtime/USDCore/Source/UnrealUSDWrapper/Public/USDIncludesStart.h
+++ b/Engine/Plugins/Runtime/USDCore/Source/UnrealUSDWrapper/Public/USDIncludesStart.h
@@ -6,6 +6,12 @@
 
 #if USE_USD_SDK
 
+#ifndef USD_INCLUDES_START
+	#define USD_INCLUDES_START
+#else
+	#error Nesting USDIncludesStart.h is not allowed!
+#endif
+
 #include "CoreMinimal.h"
 
 THIRD_PARTY_INCLUDES_START
