Fixed non-unity USD compile issues.

diff --git a/Engine/Plugins/Runtime/USDCore/Source/UnrealUSDWrapper/Public/USDIncludesEnd.h b/Engine/Plugins/Runtime/USDCore/Source/UnrealUSDWrapper/Public/USDIncludesEnd.h
--- a/Engine/Plugins/Runtime/USDCore/Source/UnrealUSDWrapper/Public/USDIncludesEnd.h
+++ b/Engine/Plugins/Runtime/USDCore/Source/UnrealUSDWrapper/Public/USDIncludesEnd.h
@@ -30,7 +30,11 @@ THIRD_PARTY_INCLUDES_END
 #undef USD_INCLUDES_START
 
 #else
+
+#if !defined(UE_DIRECT_HEADER_COMPILE)
 	#error Mismatched USDIncludesEnd.h detected.
+#endif
+
 #endif // USD_INCLUDES_START
 
 #endif // #if USE_USD_SDK
