[USD]
* Changed so merged modules path behaves like monolithic in terms of operator new/delete overloads. Not sure if usd will work but we get around linker errors for now

diff --git a/Engine/Plugins/Runtime/USDCore/Source/UnrealUSDWrapper/Public/USDMemory.h b/Engine/Plugins/Runtime/USDCore/Source/UnrealUSDWrapper/Public/USDMemory.h
--- a/Engine/Plugins/Runtime/USDCore/Source/UnrealUSDWrapper/Public/USDMemory.h
+++ b/Engine/Plugins/Runtime/USDCore/Source/UnrealUSDWrapper/Public/USDMemory.h
@@ -262,7 +262,7 @@ TSharedRef<ObjectType> MakeSharedUnreal(ArgTypes&&... Args)
 // See comment on UnrealUSDWrapper.Build.cs to understand why we disable these for monolithic builds (everything will still work,
 // they're just unnecessary)
 // clang-format off
-#if !FORCE_ANSI_ALLOCATOR && !IS_MONOLITHIC
+#if !FORCE_ANSI_ALLOCATOR && !IS_MONOLITHIC && !USD_MERGED_MODULES
 	#define REPLACEMENT_OPERATOR_NEW_AND_DELETE_USD \
 		OPERATOR_NEW_MSVC_PRAGMA void* operator new  ( size_t Size                        ) OPERATOR_NEW_THROW_SPEC      { return FUsdMemoryManager::Malloc( Size ); } \
 		OPERATOR_NEW_MSVC_PRAGMA void* operator new[]( size_t Size                        ) OPERATOR_NEW_THROW_SPEC      { return FUsdMemoryManager::Malloc( Size ); } \
@@ -280,7 +280,7 @@ TSharedRef<ObjectType> MakeSharedUnreal(ArgTypes&&... Args)
 	#define REPLACEMENT_OPERATOR_NEW_AND_DELETE_USD
 #endif
 
-#if !FORCE_ANSI_ALLOCATOR && !IS_MONOLITHIC
+#if !FORCE_ANSI_ALLOCATOR && !IS_MONOLITHIC && !USD_MERGED_MODULES
 	#define IMPLEMENT_MODULE_USD( ModuleImplClass, ModuleName ) \
 		\
 		/**/ \
diff --git a/Engine/Plugins/Runtime/USDCore/Source/UnrealUSDWrapper/UnrealUSDWrapper.Build.cs b/Engine/Plugins/Runtime/USDCore/Source/UnrealUSDWrapper/UnrealUSDWrapper.Build.cs
--- a/Engine/Plugins/Runtime/USDCore/Source/UnrealUSDWrapper/UnrealUSDWrapper.Build.cs
+++ b/Engine/Plugins/Runtime/USDCore/Source/UnrealUSDWrapper/UnrealUSDWrapper.Build.cs
@@ -169,7 +169,11 @@ namespace UnrealBuildTool.Rules
 				PublicDefinitions.Add("USD_USES_SYSTEM_MALLOC=0");
 			}
 
-			PrivateDefinitions.Add("SUPPRESS_PER_MODULE_INLINE_FILE"); // This module does not use core's standard operator new/delete overloads
+			PublicDefinitions.Add("USD_MERGED_MODULES=" + (Target.bMergeModules ? "1" : "0"));
+			if (!Target.bMergeModules)
+			{
+				PrivateDefinitions.Add("SUPPRESS_PER_MODULE_INLINE_FILE"); // This module does not use core's standard operator new/delete overloads
+			}
 		}
 
 		bool EnableUsdSdk(ReadOnlyTargetRules Target)
