USD: Fixed USD libs location on Mac and Linux.

diff --git a/Engine/Plugins/Runtime/USDCore/Source/UnrealUSDWrapper/Private/UnrealUSDWrapper.cpp b/Engine/Plugins/Runtime/USDCore/Source/UnrealUSDWrapper/Private/UnrealUSDWrapper.cpp
--- a/Engine/Plugins/Runtime/USDCore/Source/UnrealUSDWrapper/Private/UnrealUSDWrapper.cpp
+++ b/Engine/Plugins/Runtime/USDCore/Source/UnrealUSDWrapper/Private/UnrealUSDWrapper.cpp
@@ -1112,9 +1112,9 @@ public:
 #endif	  // PLATFORM_WINDOWS
 
 #ifdef USE_LIBRARIES_FROM_PLUGIN_FOLDER
-		// e.g. "../../../Engine/Plugins/Importers/USDImporter/Source/ThirdParty"
+		// e.g. "../../../Engine/Plugins/Runtime/USDCore/Source/ThirdParty"
 		FString TargetDllFolder = FPaths::Combine(
-			IPluginManager::Get().FindPlugin(TEXT("USDImporter"))->GetBaseDir(),
+			IPluginManager::Get().FindPlugin(TEXT("USDCore"))->GetBaseDir(),
 			TEXT("Source"),
 			TEXT("ThirdParty")
 		);
