Fixed a potential crash from USDAssetCache were we may use a GC'd UUsdAssetCache object from an async task. We now confirm the object is still valid before calling RescanAssetDirectory on the object.

[RN] fix Plugin/USDCore

diff --git a/Engine/Plugins/Runtime/USDCore/Source/USDClasses/Private/USDAssetCache3.cpp b/Engine/Plugins/Runtime/USDCore/Source/USDClasses/Private/USDAssetCache3.cpp
--- a/Engine/Plugins/Runtime/USDCore/Source/USDClasses/Private/USDAssetCache3.cpp
+++ b/Engine/Plugins/Runtime/USDCore/Source/USDClasses/Private/USDAssetCache3.cpp
@@ -1120,11 +1120,15 @@ void UUsdAssetCache3::PostLoad()
 	// make sure we do a new scan whenever we do get loaded to pick up on any new assets that may have been added.
 	// We delay this to the end of the frame though, because we may need to mark ourselves as dirty if we found anything,
 	// and we can't do that within the callstack that calls PostLoad on us.
+	TWeakObjectPtr<UUsdAssetCache3> This(this);
 	AsyncTask(
 		ENamedThreads::GameThread,
-		[this]()
+		[This]()
 		{
-			RescanAssetDirectory();
+			if (This.IsValid())
+			{
+				This->RescanAssetDirectory();
+			}
 		}
 	);
 }
