USD: Moved USDImporter content to USDCore.
Renamed USDImporter config to USDCore.

#virtualized

diff --git a/Engine/Plugins/Runtime/USDCore/Config/BaseUSDCore.ini b/Engine/Plugins/Runtime/USDCore/Config/BaseUSDCore.ini
new file mode 100644
--- /dev/null
+++ b/Engine/Plugins/Runtime/USDCore/Config/BaseUSDCore.ini
@@ -0,0 +1,21 @@
+[CoreRedirects]
++PackageRedirects=(OldName="/USDImporter/Blueprint/DefaultLiveLinkAnimBP", NewName="/USDCore/Blueprint/DefaultLiveLinkAnimBP")
++PackageRedirects=(OldName="/USDImporter/Materials/CardTextureMaterial", NewName="/USDCore/Materials/CardTextureMaterial")
++PackageRedirects=(OldName="/USDImporter/Materials/ChannelSelector", NewName="/USDCore/Materials/ChannelSelector")
++PackageRedirects=(OldName="/USDImporter/Materials/DisplayColor", NewName="/USDCore/Materials/DisplayColor")
++PackageRedirects=(OldName="/USDImporter/Materials/DisplayColorAndOpacity", NewName="/USDCore/Materials/DisplayColorAndOpacity")
++PackageRedirects=(OldName="/USDImporter/Materials/DisplayColorAndOpacityTwoSided", NewName="/USDCore/Materials/DisplayColorAndOpacityTwoSided")
++PackageRedirects=(OldName="/USDImporter/Materials/DisplayColorTwoSided", NewName="/USDCore/Materials/DisplayColorTwoSided")
++PackageRedirects=(OldName="/USDImporter/Materials/UsdPreviewSurface", NewName="/USDCore/Materials/UsdPreviewSurface")
++PackageRedirects=(OldName="/USDImporter/Materials/UsdPreviewSurfaceTranslucent", NewName="/USDCore/Materials/UsdPreviewSurfaceTranslucent")
++PackageRedirects=(OldName="/USDImporter/Materials/UsdPreviewSurfaceTranslucentTwoSided", NewName="/USDCore/Materials/UsdPreviewSurfaceTranslucentTwoSided")
++PackageRedirects=(OldName="/USDImporter/Materials/UsdPreviewSurfaceTranslucentTwoSidedVT", NewName="/USDCore/Materials/UsdPreviewSurfaceTranslucentTwoSidedVT")
++PackageRedirects=(OldName="/USDImporter/Materials/UsdPreviewSurfaceTranslucentVT", NewName="/USDCore/Materials/UsdPreviewSurfaceTranslucentVT")
++PackageRedirects=(OldName="/USDImporter/Materials/UsdPreviewSurfaceTwoSided", NewName="/USDCore/Materials/UsdPreviewSurfaceTwoSided")
++PackageRedirects=(OldName="/USDImporter/Materials/UsdPreviewSurfaceTwoSidedVT", NewName="/USDCore/Materials/UsdPreviewSurfaceTwoSidedVT")
++PackageRedirects=(OldName="/USDImporter/Materials/UsdPreviewSurfaceVT", NewName="/USDCore/Materials/UsdPreviewSurfaceVT")
++PackageRedirects=(OldName="/USDImporter/Materials/UsdTransform2d", NewName="/USDCore/Materials/UsdTransform2d")
++PackageRedirects=(OldName="/USDImporter/Materials/WorldSpaceNormalsToTangentSpace", NewName="/USDCore/Materials/WorldSpaceNormalsToTangentSpace")
++PackageRedirects=(OldName="/USDImporter/Textures/DefaultColorVT", NewName="/USDCore/Textures/DefaultColorVT")
++PackageRedirects=(OldName="/USDImporter/Textures/DefaultLinearColorVT", NewName="/USDCore/Textures/DefaultLinearColorVT")
++PackageRedirects=(OldName="/USDImporter/Textures/DefaultNormalVT", NewName="/USDCore/Textures/DefaultNormalVT")
diff --git a/Engine/Plugins/Runtime/USDCore/Source/USDClasses/Private/USDDrawModeComponent.cpp b/Engine/Plugins/Runtime/USDCore/Source/USDClasses/Private/USDDrawModeComponent.cpp
--- a/Engine/Plugins/Runtime/USDCore/Source/USDClasses/Private/USDDrawModeComponent.cpp
+++ b/Engine/Plugins/Runtime/USDCore/Source/USDClasses/Private/USDDrawModeComponent.cpp
@@ -954,6 +954,11 @@ UUsdDrawModeComponent::UUsdDrawModeComponent()
 	, MaterialInstances({nullptr, nullptr, nullptr, nullptr, nullptr, nullptr})
 	, AuthoredFaces{(int32)EUsdModelCardFace::None}
 {
+	if (IsTemplate())
+	{
+		return;
+	}
+
 	RefreshMaterialInstances();
 }
 
diff --git a/Engine/Plugins/Runtime/USDCore/Source/USDClasses/Public/USDProjectSettings.h b/Engine/Plugins/Runtime/USDCore/Source/USDClasses/Public/USDProjectSettings.h
--- a/Engine/Plugins/Runtime/USDCore/Source/USDClasses/Public/USDProjectSettings.h
+++ b/Engine/Plugins/Runtime/USDCore/Source/USDClasses/Public/USDProjectSettings.h
@@ -24,9 +24,10 @@ enum class EUsdEditInInstanceBehavior : uint8
 	ShowPrompt
 };
 
-// USDImporter and defaultconfig here so this ends up at DefaultUSDImporter.ini in the editor, and is sent to the
+// USDCore and defaultconfig here so this ends up at DefaultUSDCore.ini in the editor, and is sent to the
 // packaged game as well
-UCLASS(config = USDImporter, defaultconfig, meta = (DisplayName = USDImporter), MinimalAPI)
+
+UCLASS(config = USDCore, defaultconfig, meta = (DisplayName = USDCore), MinimalAPI)
 class UUsdProjectSettings : public UDeveloperSettings
 {
 	GENERATED_BODY()
@@ -127,7 +128,7 @@ public:
 	 * You can swap this with your own material, but make sure the replacement material has a "Texture" parameter
 	 */
 	UPROPERTY(config, EditAnywhere, Category = "USD|Reference Materials", meta = (AllowedClasses = "/Script/Engine.MaterialInterface"))
-	FSoftObjectPath ReferenceModelCardTextureMaterial = FSoftObjectPath{TEXT("/USDImporter/Materials/CardTextureMaterial.CardTextureMaterial")};
+	FSoftObjectPath ReferenceModelCardTextureMaterial = FSoftObjectPath{TEXT("/USDCore/Materials/CardTextureMaterial.CardTextureMaterial")};
 
 	/**
 	 * What material to use as reference material when creating material instances from USD materials.
@@ -141,7 +142,7 @@ public:
 		Category = "USD|Reference Materials|UsdPreviewSurface",
 		meta = (AllowedClasses = "/Script/Engine.MaterialInterface")
 	)
-	FSoftObjectPath ReferencePreviewSurfaceMaterial = FSoftObjectPath{TEXT("/USDImporter/Materials/UsdPreviewSurface.UsdPreviewSurface")};
+	FSoftObjectPath ReferencePreviewSurfaceMaterial = FSoftObjectPath{TEXT("/USDCore/Materials/UsdPreviewSurface.UsdPreviewSurface")};
 
 	/**
 	 * What material to use as reference material when creating material instances from USD materials.
@@ -155,7 +156,7 @@ public:
 		Category = "USD|Reference Materials|UsdPreviewSurface",
 		meta = (AllowedClasses = "/Script/Engine.MaterialInterface")
 	)
-	FSoftObjectPath ReferencePreviewSurfaceTranslucentMaterial = FSoftObjectPath{TEXT("/USDImporter/Materials/"
+	FSoftObjectPath ReferencePreviewSurfaceTranslucentMaterial = FSoftObjectPath{TEXT("/USDCore/Materials/"
 																					  "UsdPreviewSurfaceTranslucent.UsdPreviewSurfaceTranslucent")};
 
 	/**
@@ -170,7 +171,7 @@ public:
 		Category = "USD|Reference Materials|UsdPreviewSurface",
 		meta = (AllowedClasses = "/Script/Engine.MaterialInterface")
 	)
-	FSoftObjectPath ReferencePreviewSurfaceTwoSidedMaterial = FSoftObjectPath{TEXT("/USDImporter/Materials/"
+	FSoftObjectPath ReferencePreviewSurfaceTwoSidedMaterial = FSoftObjectPath{TEXT("/USDCore/Materials/"
 																				   "UsdPreviewSurfaceTwoSided.UsdPreviewSurfaceTwoSided")};
 
 	/**
@@ -186,7 +187,7 @@ public:
 		meta = (AllowedClasses = "/Script/Engine.MaterialInterface")
 	)
 	FSoftObjectPath ReferencePreviewSurfaceTranslucentTwoSidedMaterial = FSoftObjectPath{
-		TEXT("/USDImporter/Materials/UsdPreviewSurfaceTranslucentTwoSided.UsdPreviewSurfaceTranslucentTwoSided")};
+		TEXT("/USDCore/Materials/UsdPreviewSurfaceTranslucentTwoSided.UsdPreviewSurfaceTranslucentTwoSided")};
 
 	/**
 	 * What material to use as reference material when creating material instances from USD materials.
@@ -200,7 +201,7 @@ public:
 		Category = "USD|Reference Materials|UsdPreviewSurface with Virtual Textures",
 		meta = (AllowedClasses = "/Script/Engine.MaterialInterface")
 	)
-	FSoftObjectPath ReferencePreviewSurfaceVTMaterial = FSoftObjectPath{TEXT("/USDImporter/Materials/UsdPreviewSurfaceVT.UsdPreviewSurfaceVT")};
+	FSoftObjectPath ReferencePreviewSurfaceVTMaterial = FSoftObjectPath{TEXT("/USDCore/Materials/UsdPreviewSurfaceVT.UsdPreviewSurfaceVT")};
 
 	/**
 	 * What material to use as reference material when creating material instances from USD materials.
@@ -215,7 +216,7 @@ public:
 		meta = (AllowedClasses = "/Script/Engine.MaterialInterface")
 	)
 	FSoftObjectPath ReferencePreviewSurfaceTranslucentVTMaterial = FSoftObjectPath{
-		TEXT("/USDImporter/Materials/UsdPreviewSurfaceTranslucentVT.UsdPreviewSurfaceTranslucentVT")};
+		TEXT("/USDCore/Materials/UsdPreviewSurfaceTranslucentVT.UsdPreviewSurfaceTranslucentVT")};
 
 	/**
 	 * What material to use as reference material when creating material instances from USD materials.
@@ -229,7 +230,7 @@ public:
 		Category = "USD|Reference Materials|UsdPreviewSurface with Virtual Textures",
 		meta = (AllowedClasses = "/Script/Engine.MaterialInterface")
 	)
-	FSoftObjectPath ReferencePreviewSurfaceTwoSidedVTMaterial = FSoftObjectPath{TEXT("/USDImporter/Materials/"
+	FSoftObjectPath ReferencePreviewSurfaceTwoSidedVTMaterial = FSoftObjectPath{TEXT("/USDCore/Materials/"
 																					 "UsdPreviewSurfaceTwoSidedVT.UsdPreviewSurfaceTwoSidedVT")};
 
 	/**
@@ -245,7 +246,7 @@ public:
 		meta = (AllowedClasses = "/Script/Engine.MaterialInterface")
 	)
 	FSoftObjectPath ReferencePreviewSurfaceTranslucentTwoSidedVTMaterial = FSoftObjectPath{
-		TEXT("/USDImporter/Materials/UsdPreviewSurfaceTranslucentTwoSidedVT.UsdPreviewSurfaceTranslucentTwoSidedVT")};
+		TEXT("/USDCore/Materials/UsdPreviewSurfaceTranslucentTwoSidedVT.UsdPreviewSurfaceTranslucentTwoSidedVT")};
 
 	/**
 	 * What material to use as reference material when creating material instances from USD materials.
@@ -254,7 +255,7 @@ public:
 	 * the parameters filled with values extracted from the USD material when parsing.
 	 */
 	UPROPERTY(config, EditAnywhere, Category = "USD|Reference Materials|DisplayColor", meta = (AllowedClasses = "/Script/Engine.MaterialInterface"))
-	FSoftObjectPath ReferenceDisplayColorMaterial = FSoftObjectPath{TEXT("/USDImporter/Materials/DisplayColor.DisplayColor")};
+	FSoftObjectPath ReferenceDisplayColorMaterial = FSoftObjectPath{TEXT("/USDCore/Materials/DisplayColor.DisplayColor")};
 
 	/**
 	 * What material to use as reference material when creating material instances from USD materials.
@@ -263,7 +264,7 @@ public:
 	 * the parameters filled with values extracted from the USD material when parsing.
 	 */
 	UPROPERTY(config, EditAnywhere, Category = "USD|Reference Materials|DisplayColor", meta = (AllowedClasses = "/Script/Engine.MaterialInterface"))
-	FSoftObjectPath ReferenceDisplayColorAndOpacityMaterial = FSoftObjectPath{TEXT("/USDImporter/Materials/"
+	FSoftObjectPath ReferenceDisplayColorAndOpacityMaterial = FSoftObjectPath{TEXT("/USDCore/Materials/"
 																				   "DisplayColorAndOpacity.DisplayColorAndOpacity")};
 
 	/**
@@ -273,7 +274,7 @@ public:
 	 * the parameters filled with values extracted from the USD material when parsing.
 	 */
 	UPROPERTY(config, EditAnywhere, Category = "USD|Reference Materials|DisplayColor", meta = (AllowedClasses = "/Script/Engine.MaterialInterface"))
-	FSoftObjectPath ReferenceDisplayColorTwoSidedMaterial = FSoftObjectPath{TEXT("/USDImporter/Materials/DisplayColorTwoSided.DisplayColorTwoSided")};
+	FSoftObjectPath ReferenceDisplayColorTwoSidedMaterial = FSoftObjectPath{TEXT("/USDCore/Materials/DisplayColorTwoSided.DisplayColorTwoSided")};
 
 	/**
 	 * What material to use as reference material when creating material instances from USD materials.
@@ -283,5 +284,5 @@ public:
 	 */
 	UPROPERTY(config, EditAnywhere, Category = "USD|Reference Materials|DisplayColor", meta = (AllowedClasses = "/Script/Engine.MaterialInterface"))
 	FSoftObjectPath ReferenceDisplayColorAndOpacityTwoSidedMaterial = FSoftObjectPath{
-		TEXT("/USDImporter/Materials/DisplayColorAndOpacityTwoSided.DisplayColorAndOpacityTwoSided")};
+		TEXT("/USDCore/Materials/DisplayColorAndOpacityTwoSided.DisplayColorAndOpacityTwoSided")};
 };
diff --git a/Engine/Plugins/Runtime/USDCore/USDCore.uplugin b/Engine/Plugins/Runtime/USDCore/USDCore.uplugin
--- a/Engine/Plugins/Runtime/USDCore/USDCore.uplugin
+++ b/Engine/Plugins/Runtime/USDCore/USDCore.uplugin
@@ -11,7 +11,7 @@
 	"MarketplaceURL": "",
 	"SupportURL": "",
 	"EnabledByDefault": false,
-	"CanContainContent": false,
+	"CanContainContent": true,
 	"IsBetaVersion": true,
 	"Installed": false,
 	"Modules": [
