USD: Add option to export StaticMesh source data (and not only render mesh).

diff --git a/Engine/Plugins/Importers/USDImporter/Source/USDExporter/Private/StaticMeshExporterUSD.cpp b/Engine/Plugins/Importers/USDImporter/Source/USDExporter/Private/StaticMeshExporterUSD.cpp
--- a/Engine/Plugins/Importers/USDImporter/Source/USDExporter/Private/StaticMeshExporterUSD.cpp
+++ b/Engine/Plugins/Importers/USDImporter/Source/USDExporter/Private/StaticMeshExporterUSD.cpp
@@ -316,7 +316,8 @@ bool UStaticMeshExporterUsd::ExportBinary(UObject* Object, const TCHAR* Type, FA
 		UsdUtils::GetDefaultTimeCode(),
 		&AssetStage,
 		Options->MeshAssetOptions.LowestMeshLOD,
-		Options->MeshAssetOptions.HighestMeshLOD
+		Options->MeshAssetOptions.HighestMeshLOD,
+		Options->MeshAssetOptions.bExportStaticMeshSourceData
 	);
 
 	if (UE::FUsdPrim AssetDefaultPrim = AssetStage.GetDefaultPrim())
diff --git a/Engine/Plugins/Importers/USDImporter/Source/USDExporter/Private/USDAssetOptions.cpp b/Engine/Plugins/Importers/USDImporter/Source/USDExporter/Private/USDAssetOptions.cpp
--- a/Engine/Plugins/Importers/USDImporter/Source/USDExporter/Private/USDAssetOptions.cpp
+++ b/Engine/Plugins/Importers/USDImporter/Source/USDExporter/Private/USDAssetOptions.cpp
@@ -38,6 +38,7 @@ void UsdUtils::AddAnalyticsAttributes(const FUsdMeshAssetOptions& Options, TArra
 	}
 	InOutAttributes.Emplace(TEXT("LowestMeshLOD"), LexToString(Options.LowestMeshLOD));
 	InOutAttributes.Emplace(TEXT("HighestMeshLOD"), LexToString(Options.HighestMeshLOD));
+	InOutAttributes.Emplace(TEXT("ExportStaticMeshSourceData"), Options.bExportStaticMeshSourceData);
 }
 
 void UsdUtils::HashForMaterialExport(const FUsdMaterialBakingOptions& Options, FSHA1& HashToUpdate)
@@ -66,4 +67,5 @@ void UsdUtils::HashForMeshExport(const FUsdMeshAssetOptions& Options, FSHA1& Has
 
 	HashToUpdate.Update(reinterpret_cast<const uint8*>(&Options.LowestMeshLOD), sizeof(Options.LowestMeshLOD));
 	HashToUpdate.Update(reinterpret_cast<const uint8*>(&Options.HighestMeshLOD), sizeof(Options.HighestMeshLOD));
+	HashToUpdate.Update(reinterpret_cast<const uint8*>(&Options.bExportStaticMeshSourceData), sizeof(Options.bExportStaticMeshSourceData));
 }
diff --git a/Engine/Plugins/Importers/USDImporter/Source/USDExporter/Public/USDAssetOptions.h b/Engine/Plugins/Importers/USDImporter/Source/USDExporter/Public/USDAssetOptions.h
--- a/Engine/Plugins/Importers/USDImporter/Source/USDExporter/Public/USDAssetOptions.h
+++ b/Engine/Plugins/Importers/USDImporter/Source/USDExporter/Public/USDAssetOptions.h
@@ -76,6 +76,19 @@ struct USDEXPORTER_API FUsdMeshAssetOptions
 	UPROPERTY(EditAnywhere, config, BlueprintReadWrite, Category = "Material options", meta = (EditCondition = bBakeMaterials))
 	FUsdMaterialBakingOptions MaterialBakingOptions;
 
+	/**
+	 * If true it means we'll try exporting the source data when exporting static meshes.
+	 * Not all meshes contain their source data though. If we fail to find any source data, we'll fall back to exporting the render data instead.
+	 *
+	 * If false it means we'll export the render data instead: This can be closer to what is seen on the viewport, but it can also
+	 * be the most processed. For example, if the mesh has Nanite enabled this render data may be a much coarser version of the
+	 * source mesh.
+	 *
+	 * This is only relevant for StaticMeshes, and has no effect on other mesh types (SkeletalMeshes, GeometryCaches, etc.)
+	 */
+	UPROPERTY(EditAnywhere, config, BlueprintReadWrite, Category = "Mesh options")
+	bool bExportStaticMeshSourceData = true;
+
 	/** Lowest of the LOD indices to export static and skeletal meshes with (use 0 for full detail) */
 	UPROPERTY(EditAnywhere, config, BlueprintReadWrite, Category = "Mesh options", meta = (ClampMin = "0"))
 	int32 LowestMeshLOD = 0;
diff --git a/Engine/Plugins/Runtime/USDCore/Source/USDUtilities/Private/USDGeomMeshConversion.cpp b/Engine/Plugins/Runtime/USDCore/Source/USDUtilities/Private/USDGeomMeshConversion.cpp
--- a/Engine/Plugins/Runtime/USDCore/Source/USDUtilities/Private/USDGeomMeshConversion.cpp
+++ b/Engine/Plugins/Runtime/USDCore/Source/USDUtilities/Private/USDGeomMeshConversion.cpp
@@ -80,6 +80,15 @@ static FAutoConsoleVariableRef CVarMaxInstancesPerPointInstancer(
 	TEXT("We will only parse up to this many instances from any point instancer when reading from USD to UE. Set this to -1 to disable this limit.")
 );
 
+static bool GExportNaniteSourceMeshData = true;
+static FAutoConsoleVariableRef CVarExportNaniteSourceMeshData(
+	TEXT("USD.ExportNaniteSourceMeshData"),
+	GExportNaniteSourceMeshData,
+	TEXT(
+		"Try using Nanite hi-res MeshDescription data when exporting static meshes with the bExportStaticMeshSourceData option enabled. If false, it means we will always just use the StaticMesh's LOD source MeshDescription data instead."
+	)
+);
+
 static bool GIgnoreNormalsWhenSubdividing = true;
 static FAutoConsoleVariableRef CVarIgnoreNormalsWhenSubdividing(
 	TEXT("USD.Subdiv.IgnoreNormalsWhenSubdividing"),
@@ -348,12 +357,11 @@ namespace UE::UsdGeomMeshConversion::Private
 		return std::strtol(Name.c_str() + LODString.size(), EndPtr, Base);
 	}
 
-	void ConvertStaticMeshLOD(
-		int32 LODIndex,
+	bool ConvertStaticMeshLOD(
 		const FStaticMeshLODResources& LODRenderMesh,
 		pxr::UsdGeomMesh& UsdMesh,
-		const TArray<FString>& MaterialAssignments,
 		const pxr::UsdTimeCode TimeCode,
+		const TArray<FString>& MaterialAssignments,
 		pxr::UsdPrim PrimToReceiveMaterialAssignments
 	)
 	{
@@ -361,7 +369,7 @@ namespace UE::UsdGeomMeshConversion::Private
 		pxr::UsdStageRefPtr Stage = MeshPrim.GetStage();
 		if (!Stage)
 		{
-			return;
+			return false;
 		}
 		const FUsdStageInfo StageInfo{Stage};
 
@@ -600,13 +608,18 @@ namespace UE::UsdGeomMeshConversion::Private
 				}
 			}
 		}
+
+		return true;
 	}
 
 	bool ConvertMeshDescription(
 		const FMeshDescription& MeshDescription,
 		pxr::UsdGeomMesh& UsdMesh,
 		const FMatrix& AdditionalTransform,
-		const pxr::UsdTimeCode TimeCode
+		const pxr::UsdTimeCode TimeCode,
+		const TArray<FString>* MaterialIndexToContentPath = nullptr,
+		const TMap<FName, int32>* ImportedMaterialSlotNameToIndex = nullptr,
+		pxr::UsdPrim* PrimToReceiveMaterialAssignments = nullptr
 	)
 	{
 		pxr::UsdPrim MeshPrim = UsdMesh.GetPrim();
@@ -628,44 +641,138 @@ namespace UE::UsdGeomMeshConversion::Private
 		const int32 VertexInstanceCount = VertexInstanceNormals.GetNumElements();
 		const int32 FaceCount = MeshDescription.Polygons().Num();
 
-		// Points
+		//
+		// Prepare temp arrays to receive the converted data
+		//
+
+		pxr::VtArray<pxr::GfVec3f> PointsArray;
+		PointsArray.reserve(VertexCount);
+
+		pxr::VtArray<int> FaceVertexCounts;
+		FaceVertexCounts.reserve(FaceCount);
+
+		pxr::VtArray<int> FaceVertexIndices;
+		FaceVertexIndices.reserve(VertexInstanceCount);
+
+		pxr::VtArray<pxr::GfVec3f> Normals;
+		Normals.reserve(VertexInstanceCount);
+
+		pxr::VtArray<pxr::GfVec3f> DisplayColors;
+		pxr::VtArray<float> DisplayOpacities;
+
+		// Check if we'll need face-varying displayColors/opacity or not
+		// This is useful because *every* MeshDescription will have one VertexInstanceColor element for each instance
+		// with white opaque color, even if the actual source data didn't have anything. We shouldn't emit thousands of "(1, 1, 1)"
+		// to the USD file for no reason
+		bool bUseConstantColor = true;
+		FVector4f FirstColor{1.0f, 1.0f, 1.0f, 1.0f};
+		if (VertexInstanceColors.GetNumElements() > 1)
 		{
-			if (pxr::UsdAttribute Points = UsdMesh.CreatePointsAttr())
+			FirstColor = VertexInstanceColors[FVertexInstanceID(0)];
+			for (const FVertexInstanceID InstanceID : MeshDescription.VertexInstances().GetElementIDs())
 			{
-				pxr::VtArray<pxr::GfVec3f> PointsArray;
-				PointsArray.reserve(VertexCount);
-
-				for (const FVertexID VertexID : MeshDescription.Vertices().GetElementIDs())
+				const FVector4f& OtherColor = VertexInstanceColors[InstanceID];
+				if (OtherColor != FirstColor)
 				{
-					FVector UEPosition = AdditionalTransform.TransformPosition((FVector)VertexPositions[VertexID]);
-					PointsArray.push_back(UnrealToUsd::ConvertVectorFloat(StageInfo, UEPosition));
+					bUseConstantColor = false;
+					break;
 				}
-
-				Points.Set(PointsArray, TimeCode);
 			}
 		}
+		if (bUseConstantColor)
+		{
+			pxr::GfVec4f Color = UnrealToUsd::ConvertColor(FLinearColor(FirstColor));
+			DisplayColors.push_back(pxr::GfVec3f(Color[0], Color[1], Color[2]));
+			DisplayOpacities.push_back(Color[3]);
+		}
+		else
+		{
+			DisplayColors.reserve(VertexInstanceCount);
+			DisplayOpacities.reserve(VertexInstanceCount);
+		}
 
-		// Normals
+		const int32 NumUVs = VertexInstanceUVs.GetNumChannels();
+		TArray<pxr::VtVec2fArray> UVs;
+		UVs.SetNum(NumUVs);
+		for (int32 UVIndex = 0; UVIndex < NumUVs; ++UVIndex)
 		{
-			if (pxr::UsdAttribute NormalsAttribute = UsdMesh.CreateNormalsAttr())
+			UVs[UVIndex].reserve(VertexInstanceCount);
+		}
+
+		// Convert points
+		for (const FVertexID VertexID : MeshDescription.Vertices().GetElementIDs())
+		{
+			FVector UEPosition = AdditionalTransform.TransformPosition((FVector)VertexPositions[VertexID]);
+			PointsArray.push_back(UnrealToUsd::ConvertVectorFloat(StageInfo, UEPosition));
+		}
+
+		// Convert all vertex instance attributes in one go
+		//
+		// It's important to emit the polygons exactly in the order that they are in MeshDescription.Polygons() here,
+		// because down below when emitting material assignment and UsdGeomSubsets, we'll fetch the indices of these
+		// polygons when iterating over the polygon groups, and they are also meant to match the polygon order within
+		// MeshDescription.Polygons()
+		for (FPolygonID PolygonID : MeshDescription.Polygons().GetElementIDs())
+		{
+			TArray<FVertexInstanceID> PolygonVertexInstances = MeshDescription.GetPolygonVertexInstances(PolygonID);
+			FaceVertexCounts.push_back(static_cast<int>(PolygonVertexInstances.Num()));
+
+			for (const FVertexInstanceID& VertexInstanceID : PolygonVertexInstances)
 			{
-				pxr::VtArray<pxr::GfVec3f> Normals;
-				Normals.reserve(VertexInstanceCount);
+				int32 VertexIndex = MeshDescription.GetVertexInstanceVertex(VertexInstanceID).GetValue();
+				FaceVertexIndices.push_back(static_cast<int>(VertexIndex));
 
-				for (const FVertexInstanceID InstanceID : MeshDescription.VertexInstances().GetElementIDs())
+				FVector UENormal = (FVector)VertexInstanceNormals[VertexInstanceID].GetSafeNormal();
+				Normals.push_back(UnrealToUsd::ConvertVectorFloat(StageInfo, UENormal));
+
+				if (!bUseConstantColor)
 				{
-					FVector UENormal = (FVector)VertexInstanceNormals[InstanceID].GetSafeNormal();
-					Normals.push_back(UnrealToUsd::ConvertVectorFloat(StageInfo, UENormal));
+					pxr::GfVec4f Color = UnrealToUsd::ConvertColor(FLinearColor(VertexInstanceColors[VertexInstanceID]));
+					DisplayColors.push_back(pxr::GfVec3f(Color[0], Color[1], Color[2]));
+					DisplayOpacities.push_back(Color[3]);
 				}
 
-				NormalsAttribute.Set(Normals, TimeCode);
-				UsdMesh.SetNormalsInterpolation(pxr::UsdGeomTokens->faceVarying);
+				for (int32 UVIndex = 0; UVIndex < NumUVs; ++UVIndex)
+				{
+					FVector2D UV = FVector2D(VertexInstanceUVs.Get(VertexInstanceID, UVIndex));
+					UV[1] = 1.f - UV[1];
+					UVs[UVIndex].push_back(UnrealToUsd::ConvertVectorFloat(UV));
+				}
 			}
 		}
 
-		// UVs
+		// Create attributes and set converted data into USD
 		{
-			int32 NumUVs = VertexInstanceUVs.GetNumChannels();
+			pxr::UsdAttribute Points = UsdMesh.CreatePointsAttr();
+			Points.Set(PointsArray, TimeCode);
+
+			pxr::UsdAttribute FaceCountsAttribute = UsdMesh.CreateFaceVertexCountsAttr();
+			FaceCountsAttribute.Set(FaceVertexCounts, TimeCode);
+
+			pxr::UsdAttribute FaceVertexIndicesAttribute = UsdMesh.GetFaceVertexIndicesAttr();
+			FaceVertexIndicesAttribute.Set(FaceVertexIndices, TimeCode);
+
+			// We need to emit this if we're writing normals (which we always are) because any DCC that can
+			// actually subdivide (like usdview) will just discard authored normals and fully recompute them
+			// on-demand in case they have a valid subdivision scheme (which is the default state).
+			// Reference: https://graphics.pixar.com/usd/release/api/class_usd_geom_mesh.html#UsdGeom_Mesh_Normals
+			if (pxr::UsdAttribute SubdivisionAttr = UsdMesh.CreateSubdivisionSchemeAttr())
+			{
+				ensure(SubdivisionAttr.Set(pxr::UsdGeomTokens->none));
+			}
+			pxr::UsdAttribute NormalsAttribute = UsdMesh.CreateNormalsAttr();
+			NormalsAttribute.Set(Normals, TimeCode);
+			UsdMesh.SetNormalsInterpolation(pxr::UsdGeomTokens->faceVarying);
+
+			pxr::UsdGeomPrimvar DisplayColorPrimvar = UsdMesh.CreateDisplayColorPrimvar(
+				bUseConstantColor ? pxr::UsdGeomTokens->constant : pxr::UsdGeomTokens->faceVarying
+			);
+			DisplayColorPrimvar.Set(DisplayColors, TimeCode);
+
+			pxr::UsdGeomPrimvar DisplayOpacityPrimvar = UsdMesh.CreateDisplayOpacityPrimvar(
+				bUseConstantColor ? pxr::UsdGeomTokens->constant : pxr::UsdGeomTokens->faceVarying
+			);
+			DisplayOpacityPrimvar.Set(DisplayOpacities, TimeCode);
 
 			for (int32 UVIndex = 0; UVIndex < NumUVs; ++UVIndex)
 			{
@@ -673,72 +780,94 @@ namespace UE::UsdGeomMeshConversion::Private
 
 				pxr::UsdGeomPrimvar PrimvarST = pxr::UsdGeomPrimvarsAPI(MeshPrim)
 													.CreatePrimvar(UsdUVSetName, pxr::SdfValueTypeNames->TexCoord2fArray, pxr::UsdGeomTokens->vertex);
-				if (PrimvarST)
-				{
-					pxr::VtVec2fArray UVs;
 
-					for (const FVertexInstanceID InstanceID : MeshDescription.VertexInstances().GetElementIDs())
-					{
-						FVector2D UV = FVector2D(VertexInstanceUVs.Get(InstanceID, UVIndex));
-						UV[1] = 1.f - UV[1];
-						UVs.push_back(UnrealToUsd::ConvertVectorFloat(UV));
-					}
-
-					PrimvarST.Set(UVs, TimeCode);
-					PrimvarST.SetInterpolation(pxr::UsdGeomTokens->faceVarying);
-				}
+				PrimvarST.Set(UVs[UVIndex], TimeCode);
+				PrimvarST.SetInterpolation(pxr::UsdGeomTokens->faceVarying);
 			}
 		}
 
-		// Vertex colors
-		if (VertexInstanceColors.GetNumElements() > 0)
+		// Handle material assignments in another pass
+		if (MaterialIndexToContentPath && ImportedMaterialSlotNameToIndex && PrimToReceiveMaterialAssignments)
 		{
-			pxr::UsdGeomPrimvar DisplayColorPrimvar = UsdMesh.CreateDisplayColorPrimvar(pxr::UsdGeomTokens->faceVarying);
-			pxr::UsdGeomPrimvar DisplayOpacityPrimvar = UsdMesh.CreateDisplayOpacityPrimvar(pxr::UsdGeomTokens->faceVarying);
-			if (DisplayColorPrimvar && DisplayOpacityPrimvar)
+			const int32 NumPolygonGroups = MeshDescription.PolygonGroups().Num();
+			const bool bCreateSubset = NumPolygonGroups > 1;
+
+			// It's important that we're traversing the polygon groups in the same order the order used when emitting the faces,
+			// as we may need to emit triangle indices that match that data if we're emitting UsdGeomSubsets
+			for (const FPolygonGroupID& PolygonGroupID : MeshDescription.PolygonGroups().GetElementIDs())
 			{
-				pxr::VtArray<pxr::GfVec3f> DisplayColors;
-				DisplayColors.reserve(VertexInstanceCount);
+				const FName& PolygonGroupImportedSlotName = PolygonGroupImportedMaterialSlotNames[PolygonGroupID];
 
-				pxr::VtArray<float> DisplayOpacities;
-				DisplayOpacities.reserve(VertexInstanceCount);
+				int32 MaterialIndex = PolygonGroupID.GetValue();
+				if (const int32* FoundIndex = ImportedMaterialSlotNameToIndex->Find(PolygonGroupImportedSlotName))
+				{
+					MaterialIndex = *FoundIndex;
+				}
 
-				for (const FVertexInstanceID InstanceID : MeshDescription.VertexInstances().GetElementIDs())
+				const FString* ContentPath = nullptr;
+				if (MaterialIndexToContentPath->IsValidIndex(MaterialIndex))
 				{
-					pxr::GfVec4f Color = UnrealToUsd::ConvertColor(FLinearColor(VertexInstanceColors[InstanceID]));
-					DisplayColors.push_back(pxr::GfVec3f(Color[0], Color[1], Color[2]));
-					DisplayOpacities.push_back(Color[3]);
+					ContentPath = &((*MaterialIndexToContentPath)[MaterialIndex]);
 				}
 
-				DisplayColorPrimvar.Set(DisplayColors, TimeCode);
-				DisplayOpacityPrimvar.Set(DisplayOpacities, TimeCode);
-			}
-		}
+				// Create the triangles
+				if (bCreateSubset)
+				{
+					// Create an UsdGeomSubset for this polygon group, as we have multiple assignments (and sections) in this mesh to write out.
+					// Note that the subsets need to be a valid partition of the mesh, so we must create one even if we failed to find an
+					// actual material assignment to use for it
 
-		// Faces
-		{
-			pxr::UsdAttribute FaceCountsAttribute = UsdMesh.CreateFaceVertexCountsAttr();
-			pxr::UsdAttribute FaceVertexIndicesAttribute = UsdMesh.GetFaceVertexIndicesAttr();
+					const int32 SectionIndex = PolygonGroupID.GetValue();
 
-			pxr::VtArray<int> FaceVertexCounts;
-			FaceVertexCounts.reserve(FaceCount);
+					pxr::UsdPrim GeomSubsetPrim = Stage->DefinePrim(
+						MeshPrim.GetPath().AppendPath(pxr::SdfPath("Section" + std::to_string(SectionIndex))),
+						UnrealToUsd::ConvertToken(TEXT("GeomSubset")).Get()
+					);
 
-			pxr::VtArray<int> FaceVertexIndices;
+					// MaterialPrim may be in another stage, so we may need another GeomSubset there
+					pxr::UsdPrim MaterialGeomSubsetPrim = GeomSubsetPrim;
+					if (PrimToReceiveMaterialAssignments->GetStage() != MeshPrim.GetStage())
+					{
+						MaterialGeomSubsetPrim = PrimToReceiveMaterialAssignments->GetStage()->OverridePrim(
+							PrimToReceiveMaterialAssignments->GetPath().AppendPath(pxr::SdfPath("Section" + std::to_string(SectionIndex)))
+						);
+					}
 
-			for (FPolygonID PolygonID : MeshDescription.Polygons().GetElementIDs())
-			{
-				const TArray<FVertexInstanceID>& PolygonVertexInstances = MeshDescription.GetPolygonVertexInstances(PolygonID);
-				FaceVertexCounts.push_back(static_cast<int>(PolygonVertexInstances.Num()));
+					pxr::UsdGeomSubset GeomSubsetSchema{GeomSubsetPrim};
+
+					// Element type attribute
+					pxr::UsdAttribute ElementTypeAttr = GeomSubsetSchema.CreateElementTypeAttr();
+					ElementTypeAttr.Set(pxr::UsdGeomTokens->face, TimeCode);
+
+					// Indices attribute
+					pxr::VtArray<int> IndicesAttrValue;
+					for (const FPolygonID& PolygonID : MeshDescription.GetPolygonGroupPolygonIDs(PolygonGroupID))
+					{
+						IndicesAttrValue.push_back(static_cast<int>(PolygonID.GetValue()));
+					}
+
+					pxr::UsdAttribute IndicesAttr = GeomSubsetSchema.CreateIndicesAttr();
+					IndicesAttr.Set(IndicesAttrValue, TimeCode);
+
+					// Family name attribute
+					pxr::UsdAttribute FamilyNameAttr = GeomSubsetSchema.CreateFamilyNameAttr();
+					FamilyNameAttr.Set(pxr::UsdShadeTokens->materialBind, TimeCode);
+
+					// Family type
+					pxr::UsdGeomSubset::SetFamilyType(UsdMesh, pxr::UsdShadeTokens->materialBind, pxr::UsdGeomTokens->partition);
 
-				for (FVertexInstanceID VertexInstanceID : PolygonVertexInstances)
+					// material:binding relationship
+					if (ContentPath)
+					{
+						UsdUtils::AuthorUnrealMaterialBinding(MaterialGeomSubsetPrim, *ContentPath);
+					}
+				}
+				else if (ContentPath)
 				{
-					int32 VertexIndex = MeshDescription.GetVertexInstanceVertex(VertexInstanceID).GetValue();
-					FaceVertexIndices.push_back(static_cast<int>(VertexIndex));
+					// This will be the single assignment for this entire mesh: Just create the assignment directly on the Mesh
+					UsdUtils::AuthorUnrealMaterialBinding(*PrimToReceiveMaterialAssignments, *ContentPath);
 				}
 			}
-
-			FaceCountsAttribute.Set(FaceVertexCounts, TimeCode);
-			FaceVertexIndicesAttribute.Set(FaceVertexIndices, TimeCode);
 		}
 
 		return true;
@@ -3593,7 +3722,8 @@ bool UnrealToUsd::ConvertStaticMesh(
 	const pxr::UsdTimeCode TimeCode,
 	UE::FUsdStage* StageForMaterialAssignments,
 	int32 LowestMeshLOD,
-	int32 HighestMeshLOD
+	int32 HighestMeshLOD,
+	bool bExportStaticMeshSourceData
 )
 {
 	FScopedUsdAllocs UsdAllocs;
@@ -3655,6 +3785,7 @@ bool UnrealToUsd::ConvertStaticMesh(
 #endif	  // WITH_EDITOR
 
 	pxr::UsdVariantSets VariantSets = UsdPrim.GetVariantSets();
+	pxr::UsdVariantSet VariantSet = VariantSets.GetVariantSet(UnrealIdentifiers::LOD);
 	if (NumLODs > 1 && VariantSets.HasVariantSet(UnrealIdentifiers::LOD))
 	{
 		UE_LOG(
@@ -3674,9 +3805,13 @@ bool UnrealToUsd::ConvertStaticMesh(
 
 	// Collect all material assignments, referenced by the sections' material indices
 	bool bHasMaterialAssignments = false;
-	TArray<FString> MaterialAssignments;
-	for (const FStaticMaterial& StaticMaterial : StaticMesh->GetStaticMaterials())
+	TArray<FString> MaterialIndexToContentPath;			   // This one is used when exporting static mesh render data (we retain the order)
+	TMap<FName, int32> ImportedMaterialSlotNameToIndex;	   // This is used when exporting FMeshDescriptions (we use the imported slot names)
+	const TArray<FStaticMaterial>& StaticMaterials = StaticMesh->GetStaticMaterials();
+	for (int32 MaterialIndex = 0; MaterialIndex < StaticMaterials.Num(); ++MaterialIndex)
 	{
+		const FStaticMaterial& StaticMaterial = StaticMaterials[MaterialIndex];
+
 		FString AssignedMaterialPathName;
 		if (UMaterialInterface* Material = StaticMaterial.MaterialInterface)
 		{
@@ -3687,13 +3822,17 @@ bool UnrealToUsd::ConvertStaticMesh(
 			}
 		}
 
-		MaterialAssignments.Add(AssignedMaterialPathName);
+		MaterialIndexToContentPath.Add(AssignedMaterialPathName);
+		ImportedMaterialSlotNameToIndex.Add(StaticMaterial.ImportedMaterialSlotName, MaterialIndex);
 	}
 	if (!bHasMaterialAssignments)
 	{
 		// Prevent creation of the UnrealMaterials prims in case we don't have any assignments at all
-		MaterialAssignments.Reset();
+		MaterialIndexToContentPath.Reset();
+		ImportedMaterialSlotNameToIndex.Reset();
 	}
+	// Author material bindings on the dedicated stage if we have one
+	pxr::UsdStageRefPtr MaterialStage = StageForMaterialAssignments ? static_cast<pxr::UsdStageRefPtr>(*StageForMaterialAssignments) : Stage;
 
 	// Do this outside the variant edit context or else it's going to be a weaker opinion than the stuff outside
 	// the variant, and it won't really do anything for UsdPrim if it already exists.
@@ -3702,21 +3841,12 @@ bool UnrealToUsd::ConvertStaticMesh(
 	// here, so would our referencer and we wouldn't be able to put a transform on it
 	UsdPrim = Stage->DefinePrim(UsdPrim.GetPath(), UnrealToUsd::ConvertToken(bExportMultipleLODs ? TEXT("Xform") : TEXT("Mesh")).Get());
 
+	const bool bExportNaniteDataAsSourceData = GExportNaniteSourceMeshData && StaticMesh->IsNaniteEnabled()
+											   && StaticMesh->IsHiResMeshDescriptionValid();
+
+	bool bExported = false;
 	for (int32 LODIndex = LowestMeshLOD; LODIndex <= HighestMeshLOD; ++LODIndex)
 	{
-		const FStaticMeshLODResources& RenderMesh = StaticMesh->GetLODForExport(LODIndex);
-
-		// Verify the integrity of the static mesh.
-		if (RenderMesh.VertexBuffers.StaticMeshVertexBuffer.GetNumVertices() == 0)
-		{
-			continue;
-		}
-
-		if (RenderMesh.Sections.Num() == 0)
-		{
-			continue;
-		}
-
 		// LOD0, LOD1, etc
 		std::string VariantName = UnrealIdentifiers::LOD.GetString() + UnrealToUsd::ConvertString(*LexToString(LODIndex)).Get();
 		if (LowestLODAdded.size() == 0)
@@ -3724,14 +3854,10 @@ bool UnrealToUsd::ConvertStaticMesh(
 			LowestLODAdded = VariantName;
 		}
 
-		pxr::SdfPath LODPrimPath = ParentPrimPath.AppendPath(pxr::SdfPath(VariantName));
-
 		// Enable the variant edit context, if we are creating variant LODs
 		TOptional<pxr::UsdEditContext> EditContext;
 		if (bExportMultipleLODs)
 		{
-			pxr::UsdVariantSet VariantSet = VariantSets.GetVariantSet(UnrealIdentifiers::LOD);
-
 			if (!VariantSet.AddVariant(VariantName))
 			{
 				continue;
@@ -3741,16 +3867,7 @@ bool UnrealToUsd::ConvertStaticMesh(
 			EditContext.Emplace(VariantSet.GetVariantEditContext());
 		}
 
-		// Author material bindings on the dedicated stage if we have one
-		pxr::UsdStageRefPtr MaterialStage;
-		if (StageForMaterialAssignments)
-		{
-			MaterialStage = static_cast<pxr::UsdStageRefPtr>(*StageForMaterialAssignments);
-		}
-		else
-		{
-			MaterialStage = Stage;
-		}
+		pxr::SdfPath LODPrimPath = ParentPrimPath.AppendPath(pxr::SdfPath(VariantName));
 
 		pxr::UsdGeomMesh TargetMesh;
 		pxr::UsdPrim MaterialPrim = UsdPrim;
@@ -3769,7 +3886,73 @@ bool UnrealToUsd::ConvertStaticMesh(
 			MaterialPrim = MaterialStage->OverridePrim(UsdPrim.GetPath());
 		}
 
-		UsdGeomMeshImpl::ConvertStaticMeshLOD(LODIndex, RenderMesh, TargetMesh, MaterialAssignments, TimeCode, MaterialPrim);
+		// Try exporting source data
+		//  - Reference: FFbxExporter::ExportStaticMeshToFBX
+		const bool bUseNaniteData = LODIndex == 0 && bExportNaniteDataAsSourceData && StaticMesh->IsNaniteEnabled()
+									&& StaticMesh->IsHiResMeshDescriptionValid();
+		const bool bHasSourceData = StaticMesh->IsMeshDescriptionValid(LODIndex);
+		bool bExportSourceForLOD = bExportStaticMeshSourceData && (bUseNaniteData || bHasSourceData);
+
+		if (bExportSourceForLOD)
+		{
+			FMatrix AdditionalTransform = FTransform::Identity.ToMatrixWithScale();
+			if (bUseNaniteData)
+			{
+				if (const FMeshDescription* MeshDescription = StaticMesh->GetHiResMeshDescription())
+				{
+					bExported = UsdGeomMeshImpl::ConvertMeshDescription(
+						*MeshDescription,
+						TargetMesh,
+						AdditionalTransform,
+						TimeCode,
+						&MaterialIndexToContentPath,
+						&ImportedMaterialSlotNameToIndex,
+						&MaterialPrim
+					);
+				}
+			}
+			else
+			{
+				if (const FMeshDescription* MeshDescription = StaticMesh->GetMeshDescription(LODIndex))
+				{
+					bExported = UsdGeomMeshImpl::ConvertMeshDescription(
+						*MeshDescription,
+						TargetMesh,
+						AdditionalTransform,
+						TimeCode,
+						&MaterialIndexToContentPath,
+						&ImportedMaterialSlotNameToIndex,
+						&MaterialPrim
+					);
+				}
+			}
+		}
+		else
+		{
+			// If we want to export the render data, get it and check its integrity
+			const FStaticMeshLODResources* RenderMesh = &StaticMesh->GetLODForExport(LODIndex);
+			if (!RenderMesh || RenderMesh->VertexBuffers.StaticMeshVertexBuffer.GetNumVertices() == 0 || RenderMesh->Sections.Num() == 0)
+			{
+				UE_LOG(LogUsd, Warning, TEXT("Found invalid render data for LOD '%d' of '%s'!"), LODIndex, *StaticMesh->GetPathName());
+				continue;
+			}
+
+			// Export render data LOD
+			bExported = UsdGeomMeshImpl::ConvertStaticMeshLOD(*RenderMesh, TargetMesh, TimeCode, MaterialIndexToContentPath, MaterialPrim);
+		}
+
+		if (!bExported)
+		{
+			UE_LOG(
+				LogUsd,
+				Warning,
+				TEXT("Failed to export LOD '%d' of mesh '%s' onto prim '%s'!"),
+				LODIndex,
+				*StaticMesh->GetPathName(),
+				*UsdToUnreal::ConvertPath(TargetMesh.GetPrim().GetPrimPath())
+			);
+			break;
+		}
 	}
 
 	// Reset variant set to start with the lowest lod selected
@@ -3778,7 +3961,7 @@ bool UnrealToUsd::ConvertStaticMesh(
 		VariantSets.GetVariantSet(UnrealIdentifiers::LOD).SetVariantSelection(LowestLODAdded);
 	}
 
-	return true;
+	return bExported;
 }
 
 bool UnrealToUsd::ConvertMeshDescriptions(
@@ -4535,7 +4718,7 @@ bool UsdUtils::IterateLODMeshes(const pxr::UsdPrim& ParentPrim, TFunction<bool(c
 	const std::string OriginalVariant = LODVariantSet.GetVariantSelection();
 
 	pxr::UsdStageRefPtr Stage = ParentPrim.GetStage();
-	pxr::UsdEditContext(Stage, Stage->GetRootLayer());
+	pxr::UsdEditContext EditContext{Stage, Stage->GetRootLayer()};
 
 	bool bHasValidVariant = false;
 	for (const std::string& LODVariantName : VariantSets.GetVariantSet(LODString).GetVariantNames())
diff --git a/Engine/Plugins/Runtime/USDCore/Source/USDUtilities/Public/USDGeomMeshConversion.h b/Engine/Plugins/Runtime/USDCore/Source/USDUtilities/Public/USDGeomMeshConversion.h
--- a/Engine/Plugins/Runtime/USDCore/Source/USDUtilities/Public/USDGeomMeshConversion.h
+++ b/Engine/Plugins/Runtime/USDCore/Source/USDUtilities/Public/USDGeomMeshConversion.h
@@ -205,6 +205,7 @@ namespace UnrealToUsd
 	 * layer, but the material assignments to an asset layer)
 	 * @param LowestMeshLOD - Lowest LOD of the UStaticMesh to export (start of the LOD range)
 	 * @param HighestMeshLOD - Lowest LOD of the UStaticMesh to export (end of the LOD range)
+	 * @param bExportStaticMeshSourceData - Whether to export the source mesh description data (if true) or the render data (if false)
 	 * @return Whether the conversion was successful or not.
 	 */
 	USDUTILITIES_API bool ConvertStaticMesh(
@@ -213,7 +214,8 @@ namespace UnrealToUsd
 		pxr::UsdTimeCode TimeCode = pxr::UsdTimeCode::Default(),
 		UE::FUsdStage* StageForMaterialAssignments = nullptr,
 		int32 LowestMeshLOD = 0,
-		int32 HighestMeshLOD = INT32_MAX
+		int32 HighestMeshLOD = INT32_MAX,
+		bool bExportStaticMeshSourceData = false
 	);
 
 	/**
