USD: Fix build breakage for non-editor targets when trying to use editor-only code.

Horde issue 489744.

diff --git a/Engine/Plugins/Runtime/USDCore/Source/USDUtilities/Private/USDGeomMeshConversion.cpp b/Engine/Plugins/Runtime/USDCore/Source/USDUtilities/Private/USDGeomMeshConversion.cpp
--- a/Engine/Plugins/Runtime/USDCore/Source/USDUtilities/Private/USDGeomMeshConversion.cpp
+++ b/Engine/Plugins/Runtime/USDCore/Source/USDUtilities/Private/USDGeomMeshConversion.cpp
@@ -3886,6 +3886,7 @@ bool UnrealToUsd::ConvertStaticMesh(
 			MaterialPrim = MaterialStage->OverridePrim(UsdPrim.GetPath());
 		}
 
+#if WITH_EDITOR
 		// Try exporting source data
 		//  - Reference: FFbxExporter::ExportStaticMeshToFBX
 		const bool bUseNaniteData = LODIndex == 0 && bExportNaniteDataAsSourceData && StaticMesh->IsNaniteEnabled()
@@ -3928,6 +3929,7 @@ bool UnrealToUsd::ConvertStaticMesh(
 			}
 		}
 		else
+#endif	  // WITH_EDITOR
 		{
 			// If we want to export the render data, get it and check its integrity
 			const FStaticMeshLODResources* RenderMesh = &StaticMesh->GetLODForExport(LODIndex);
