USD: Add support for the spatial audio schema via USoundWave assets, UAudioComponent and AAmbientSound;

Fix LevelSequenceHelper using stale assets after an USD object change;
Fix USD Stage Editor not updating after a LevelSequence change;
Fix LevelSequence export also emitting empty/dummy overs for components like billboards, arrows, etc.;

#virtualized

diff --git a/Engine/Plugins/Importers/USDImporter/Source/USDExporter/Private/LevelSequenceExporterUSD.cpp b/Engine/Plugins/Importers/USDImporter/Source/USDExporter/Private/LevelSequenceExporterUSD.cpp
--- a/Engine/Plugins/Importers/USDImporter/Source/USDExporter/Private/LevelSequenceExporterUSD.cpp
+++ b/Engine/Plugins/Importers/USDImporter/Source/USDExporter/Private/LevelSequenceExporterUSD.cpp
@@ -40,11 +40,13 @@
 #include "MovieSceneSpawnRegister.h"
 #include "MovieSceneTimeHelpers.h"
 #include "MovieSceneTrack.h"
+#include "Sections/MovieSceneAudioSection.h"
 #include "Sections/MovieSceneSubSection.h"
 #include "Selection.h"
 #include "Sequencer/MovieSceneControlRigParameterTrack.h"
 #include "Subsystems/AssetEditorSubsystem.h"
 #include "Tracks/MovieScene3DAttachTrack.h"
+#include "Tracks/MovieSceneAudioTrack.h"
 #include "Tracks/MovieScenePropertyTrack.h"
 #include "Tracks/MovieSceneSkeletalAnimationTrack.h"
 #include "Tracks/MovieSceneSpawnTrack.h"
@@ -524,15 +526,18 @@ namespace UE::LevelSequenceExporterUSD::Private
 		return false;
 	}
 
-	TMap<UMovieSceneSequence*, TArray<FMovieSceneSequenceID>> GetSequenceHierarchyInstances(UMovieSceneSequence& Sequence, ISequencer& Sequencer)
+	TMap<UMovieSceneSequence*, TArray<FMovieSceneSequenceID>> GetSequenceHierarchyInstances(
+		UMovieSceneSequence& Sequence,
+		ISequencer& Sequencer,
+		FMovieSceneSequenceHierarchy& InOutHierarchyCache
+	)
 	{
 		TMap<UMovieSceneSequence*, TArray<FMovieSceneSequenceID>> SequenceInstances;
 
-		FMovieSceneSequenceHierarchy SequenceHierarchyCache;
-		UMovieSceneCompiledDataManager::CompileHierarchy(&Sequence, &SequenceHierarchyCache, EMovieSceneServerClientMask::All);
+		UMovieSceneCompiledDataManager::CompileHierarchy(&Sequence, &InOutHierarchyCache, EMovieSceneServerClientMask::All);
 
 		SequenceInstances.FindOrAdd(&Sequence).Add(Sequencer.GetRootTemplateID());
-		for (const TTuple<FMovieSceneSequenceID, FMovieSceneSubSequenceData>& Pair : SequenceHierarchyCache.AllSubSequenceData())
+		for (const TTuple<FMovieSceneSequenceID, FMovieSceneSubSequenceData>& Pair : InOutHierarchyCache.AllSubSequenceData())
 		{
 			if (ULevelSequence* SubSequence = Cast<ULevelSequence>(Pair.Value.GetSequence()))
 			{
@@ -549,9 +554,11 @@ namespace UE::LevelSequenceExporterUSD::Private
 		UMovieSceneSequence& RootSequence
 	)
 	{
+		FMovieSceneSequenceHierarchy HierarchyCache;
 		TMap<UMovieSceneSequence*, TArray<FMovieSceneSequenceID>> SequenceInstances = GetSequenceHierarchyInstances(
 			RootSequence,
-			Context.Sequencer.Get()
+			Context.Sequencer.Get(),
+			HierarchyCache
 		);
 
 		UMovieSceneSequence* OrigRootSequence = Context.Sequencer->GetRootMovieSceneSequence();
@@ -778,6 +785,70 @@ namespace UE::LevelSequenceExporterUSD::Private
 		ActorSelection->EndBatchSelectOperation(bNotify);
 	}
 
+	// Export the provided AudioTrack to Prim as UsdMediaSpatialAudio attributes
+	//
+	// Exporting this track type takes a different approach because unlike all other animation types, there is no actual
+	// change on the component on the level while the LevelSequence plays audio. This means that reading the "final output"
+	// of the sequence on the component every EvalFrame with a baker (like all other track cases do) doesn't really do
+	// anything for us, and we actually need to traverse the Sequencer tracks themselves.
+	//
+	// Of course, we won't get the benefit of the previous approach here: If we have multiple audio tracks for the same
+	// actor/component, even if they're placed within different Subsequences, they *will* conflict on the USD files.
+	// There is not much we can do about that at this point other than to emit a warning, but hopefully having multiple
+	// audio tracks on the same audio component is something that doesn't happen very often in practice anyway. If that
+	// is ever requested, in the future we could handle it by creating a separate UsdMediaSpatialAudio prim per audio section,
+	// but that will make a bit of a mess and harm roundtripping, so for now we only handle one section.
+	void ExportAudioTrack(
+		const UMovieSceneAudioTrack& AudioTrack,
+		const FMovieSceneSequenceTransform& SequenceTransform,
+		UE::FUsdPrim& Prim,
+		TMap<FString, int32>& AudioTracksPerPrim
+	)
+	{
+		FString PrimPath = Prim.GetPrimPath().GetString();
+
+		const TArray<UMovieSceneSection*>& Sections = AudioTrack.GetAudioSections();
+		if (Sections.Num() > 1)
+		{
+			// We only support one audio section per track because we need a full UsdMediaSpatialAudio prim for each
+			// section. If we tried exporting another section here we'd need a fully separate prim for it, which opens
+			// a can of worms as we so far only had one prim per binding. Furthermore we'd need to pay attention to this
+			// split when opening the stage as well, otherwise we'd roundtrip the audio track incorrectly
+			UE_LOG(
+				LogUsd,
+				Warning,
+				TEXT("The audio track '%s' has %d sections, but only the first audio section of an audio track can be written out to USD for now"),
+				*AudioTrack.GetPathName(),
+				Sections.Num()
+			);
+		}
+
+		if (Sections.Num() > 0)
+		{
+			if (UMovieSceneAudioSection* AudioSection = Cast<UMovieSceneAudioSection>(Sections[0]))
+			{
+				UnrealToUsd::ConvertAudioSection(*AudioSection, SequenceTransform, Prim);
+				AudioTracksPerPrim.FindOrAdd(PrimPath) += 1;
+			}
+		}
+
+		if (int32* SourceTracksForPrim = AudioTracksPerPrim.Find(PrimPath))
+		{
+			if (*SourceTracksForPrim > 1)
+			{
+				UE_LOG(
+					LogUsd,
+					Warning,
+					TEXT(
+						"Exporting multiple audio tracks (like '%s') to the same prim ('%s') is currently not supported and may lead to incorrect output"
+					),
+					*AudioTrack.GetPathName(),
+					*PrimPath
+				);
+			}
+		}
+	}
+
 	// Appends to InOutComponentBakers all of the component bakers for all components bound to MovieSceneSequence.
 	// In the process it will generate the output prims for each of these components, and keep track of them
 	// within the bakers themselves
@@ -785,6 +856,7 @@ namespace UE::LevelSequenceExporterUSD::Private
 		FLevelSequenceExportContext& Context,
 		UMovieSceneSequence& MovieSceneSequence,
 		const TMap<UMovieSceneSequence*, TArray<FMovieSceneSequenceID>>& SequenceInstances,
+		const FMovieSceneSequenceHierarchy& HierarchyCache,
 		UE::FUsdStage& UsdStage,
 		TMap<USceneComponent*, FCombinedComponentBakers>& InOutComponentBakers
 	)
@@ -931,6 +1003,12 @@ namespace UE::LevelSequenceExporterUSD::Private
 
 						for (USceneComponent* Child : Children)
 						{
+							// Skip hidden billboards/arrows/camera mesh components, etc.
+							if (!Child || !Child->IsVisibleInEditor() || Child->IsVisualizationComponent())
+							{
+								continue;
+							}
+
 							NewEntries.Add(Child, FSpawnedInstanceKey{});
 						}
 					}
@@ -1099,6 +1177,7 @@ namespace UE::LevelSequenceExporterUSD::Private
 
 			bool bHasTransformBaker = false;
 			bool bHasSkeletalBaker = false;
+			TMap<FString, int32> AudioTracksPerPrim;
 			if (const FMovieSceneBinding* Binding = MovieScene->FindBinding(InstanceKey.Key))
 			{
 				for (const UMovieSceneTrack* Track : Binding->GetTracks())
@@ -1137,6 +1216,28 @@ namespace UE::LevelSequenceExporterUSD::Private
 					{
 						UnrealToUsd::CreateComponentPropertyBaker(Prim, *BoundComponent, TEXT("Transform"), Baker);
 					}
+					else if (const UMovieSceneAudioTrack* AudioTrack = Cast<UMovieSceneAudioTrack>(Track))
+					{
+						FMovieSceneSequenceID InstanceID = Context.Sequencer->GetRootTemplateID();
+						if (InstancesOfThisSequence && InstancesOfThisSequence->Num() > 0)
+						{
+							InstanceID = (*InstancesOfThisSequence)[0];
+						}
+
+						FMovieSceneSequenceTransform SequenceTransform;
+						if (const FMovieSceneSubSequenceData* SubSequenceData = HierarchyCache.FindSubData(InstanceID))
+						{
+							SequenceTransform = SubSequenceData->RootToSequenceTransform;
+						}
+
+						// This is awkwardly handled here within GenerateBakersForMovieScene (even though it doesn't generate a baker) for two
+						// reasons:
+						//  - It's the first place you'd go do in order to search for how audio is exported, since literally every other type
+						//    of track we support goes through here
+						//  - Getting the Prim to export the audio track *to* is very much non-trivial and requires looking into
+						//    DynamicBindings and etc., which this function already does
+						ExportAudioTrack(*AudioTrack, SequenceTransform, Prim, AudioTracksPerPrim);
+					}
 
 					AddBaker(Baker);
 
@@ -1246,7 +1347,7 @@ namespace UE::LevelSequenceExporterUSD::Private
 					continue;
 				}
 
-				GenerateBakersForMovieScene(Context, *SubSequence, SequenceInstances, UsdStage, InOutComponentBakers);
+				GenerateBakersForMovieScene(Context, *SubSequence, SequenceInstances, HierarchyCache, UsdStage, InOutComponentBakers);
 			}
 		}
 	}
@@ -1551,13 +1652,15 @@ namespace UE::LevelSequenceExporterUSD::Private
 		// same result as if we had exported that subsequence's LevelSequence by itself
 		Context.Sequencer->ResetToNewRootSequence(MovieSceneSequence);
 
+		FMovieSceneSequenceHierarchy HierarchyCache;
 		TMap<UMovieSceneSequence*, TArray<FMovieSceneSequenceID>> SequenceInstances = GetSequenceHierarchyInstances(
 			MovieSceneSequence,
-			Context.Sequencer.Get()
+			Context.Sequencer.Get(),
+			HierarchyCache
 		);
 
 		TMap<USceneComponent*, FCombinedComponentBakers> Bakers;
-		GenerateBakersForMovieScene(Context, MovieSceneSequence, SequenceInstances, UsdStage, Bakers);
+		GenerateBakersForMovieScene(Context, MovieSceneSequence, SequenceInstances, HierarchyCache, UsdStage, Bakers);
 
 		// Bake this MovieScene
 		// We bake each MovieScene individually instead of doing one large simultaneous bake because this way
diff --git a/Engine/Plugins/Importers/USDImporter/Source/USDExporter/Private/USDConversionBlueprintContext.cpp b/Engine/Plugins/Importers/USDImporter/Source/USDExporter/Private/USDConversionBlueprintContext.cpp
--- a/Engine/Plugins/Importers/USDImporter/Source/USDExporter/Private/USDConversionBlueprintContext.cpp
+++ b/Engine/Plugins/Importers/USDImporter/Source/USDExporter/Private/USDConversionBlueprintContext.cpp
@@ -249,6 +249,21 @@ bool UUsdConversionBlueprintContext::ConvertDrawModeComponent(const UUsdDrawMode
 #endif	  // USE_USD_SDK
 }
 
+bool UUsdConversionBlueprintContext::ConvertAudioComponent(const UAudioComponent* Component, const FString& PrimPath, float TimeCode, bool bFilePathOnly)
+{
+#if USE_USD_SDK
+	UE::FUsdPrim Prim = UnrealToUsdImpl::GetPrim(Stage, PrimPath);
+	if (!Prim || !Component)
+	{
+		return false;
+	}
+
+	return UnrealToUsd::ConvertAudioComponent(*Component, Prim, bFilePathOnly, TimeCode == FLT_MAX ? UsdUtils::GetDefaultTimeCode() : TimeCode);
+#else
+	return false;
+#endif	  // USE_USD_SDK
+}
+
 bool UUsdConversionBlueprintContext::ConvertSceneComponent(const USceneComponent* Component, const FString& PrimPath)
 {
 #if USE_USD_SDK
diff --git a/Engine/Plugins/Importers/USDImporter/Source/USDExporter/Public/USDConversionBlueprintContext.h b/Engine/Plugins/Importers/USDImporter/Source/USDExporter/Public/USDConversionBlueprintContext.h
--- a/Engine/Plugins/Importers/USDImporter/Source/USDExporter/Public/USDConversionBlueprintContext.h
+++ b/Engine/Plugins/Importers/USDImporter/Source/USDExporter/Public/USDConversionBlueprintContext.h
@@ -16,6 +16,7 @@
 
 class AInstancedFoliageActor;
 class ALandscapeProxy;
+class UAudioComponent;
 class UCineCameraComponent;
 class UDirectionalLightComponent;
 class UHierarchicalInstancedStaticMeshComponent;
@@ -28,8 +29,8 @@ class USceneComponent;
 class USkyLightComponent;
 class USpotLightComponent;
 class UUsdDrawModeComponent;
-struct FUsdPrimMetadata;
 struct FUsdCombinedPrimMetadata;
+struct FUsdPrimMetadata;
 
 /**
  * Wraps the UnrealToUsd component conversion functions from the USDUtilities module so that they can be used by
@@ -123,6 +124,9 @@ public:
 	UFUNCTION(BlueprintCallable, Category = "Component conversion")
 	bool ConvertDrawModeComponent(const UUsdDrawModeComponent* Component, const FString& PrimPath, float TimeCode = 3.402823466e+38F);
 
+	UFUNCTION(BlueprintCallable, Category = "Component conversion")
+	bool ConvertAudioComponent(const UAudioComponent* Component, const FString& PrimPath, float TimeCode = 3.402823466e+38F, bool bFilePathOnly = false);
+
 	UFUNCTION(BlueprintCallable, Category = "Component conversion")
 	bool ConvertSceneComponent(const USceneComponent* Component, const FString& PrimPath);
 
diff --git a/Engine/Plugins/Importers/USDImporter/Source/USDSchemas/Private/USDMediaSpatialAudioTranslator.cpp b/Engine/Plugins/Importers/USDImporter/Source/USDSchemas/Private/USDMediaSpatialAudioTranslator.cpp
new file mode 100644
--- /dev/null
+++ b/Engine/Plugins/Importers/USDImporter/Source/USDSchemas/Private/USDMediaSpatialAudioTranslator.cpp
@@ -0,0 +1,330 @@
+// Copyright Epic Games, Inc. All Rights Reserved.
+
+#include "USDMediaSpatialAudioTranslator.h"
+
+#include "USDAssetUserData.h"
+#include "USDConversionUtils.h"
+#include "USDDrawModeComponent.h"
+#include "USDLog.h"
+#include "USDObjectUtils.h"
+#include "USDPrimConversion.h"
+#include "USDShadeConversion.h"
+#include "USDTypesConversion.h"
+#include "UsdWrappers/SdfPath.h"
+#include "UsdWrappers/UsdPrim.h"
+
+#include "Components/AudioComponent.h"
+#include "Misc/FileHelper.h"
+#include "Misc/Paths.h"
+#include "Sound/SoundWave.h"
+
+#if WITH_EDITOR
+#include "Factories/SoundFactory.h"
+#include "Utils.h"
+#elif 0
+#include "SoundFileIO/SoundFileIO.h"
+#endif	  // WITH_EDITOR
+
+#if USE_USD_SDK
+
+#include "USDIncludesStart.h"
+#include <pxr/usd/usd/prim.h>
+#include <pxr/usd/usd/stage.h>
+#include <pxr/usd/usdMedia/spatialAudio.h>
+#include "USDIncludesEnd.h"
+
+namespace UE::UsdMediaSpatialAudioTranslator::Private
+{
+	USoundWave* CreateSoundWave(const FString& FilePath, FName AssetName, UObject* Outer, EObjectFlags Flags)
+	{
+		USoundWave* SoundWave = nullptr;
+
+#if WITH_EDITOR
+		// Go through the factory if we can
+		USoundFactory* SoundWaveFactory = NewObject<USoundFactory>();
+
+		// Setup sane defaults for importing localized sound waves
+		SoundWaveFactory->bAutoCreateCue = false;
+		SoundWaveFactory->SuppressImportDialogs();
+
+		SoundWave = ImportObject<USoundWave>(Outer, AssetName, Flags, *FilePath, nullptr, SoundWaveFactory);
+
+#elif 0
+		// Replicate what the Factory is doing without some extras like cue assets, channel number conversion,
+		// AmbiX or FuMa tags for now.
+		//
+		// Note: This branch is currently disabled as it seems no longer possible to setup these SoundWave assets
+		// at runtime. This uses code that supposedly worked well in the past so it's not yet clear whether the
+		// issues are with this approach or the audio code downstream itself...
+		// For the moment, we'll just support audio import with the Editor, and not at Runtime.
+		// When enabling this, note that you'd also need a dependency on the AudioMixer module.
+		//
+		// References:
+		//  - USoundFactory::FactoryCreateBinary
+		//  - USoundFactory::CreateObject
+		//  - UInterchangeSoundWaveFactory::ImportAsset_Async
+
+		// Read file into a byte buffer
+		TArray<uint8> FileBytes;
+		if (!FFileHelper::LoadFileToArray(FileBytes, *FilePath))
+		{
+			UE_LOG(LogUsd, Warning, TEXT("Failed to load audio file '%s' into array"), *FilePath);
+			return nullptr;
+		}
+
+		// I don't quite understand why, but having extra capacity in this buffer will cause
+		// Audio::SoundFileUtils::ConvertAudioToWav below to trigger some out of bounds access checks.
+		// The code path via the SoundFactory "avoids" this problem (unintentionally?) because it reads
+		// the file bytes into a TArray64<uint8>, and then memcopies into a TArray<uint8> with the exact
+		// size and capacity it needs (ending up with no extra slack)
+		FileBytes.Shrink();
+
+		// If it's another file format other than .wav, try converting into wav data
+		FString Extension = FPaths::GetExtension(FilePath);
+		if (Extension.ToLower() != TEXT("wav"))
+		{
+			// Convert audio data to a wav file in memory
+			TArray<uint8> RawWaveData;
+			if (Audio::SoundFileUtils::ConvertAudioToWav(FileBytes, RawWaveData))
+			{
+				Swap(FileBytes, RawWaveData);
+			}
+			else
+			{
+				UE_LOG(LogUsd, Warning, TEXT("Failed to convert audio file format '%s' to wav"), *Extension);
+				return nullptr;
+			}
+		}
+
+		// Read WAV info about file
+		FWaveModInfo WaveInfo;
+		FString ErrorMessage;
+		if (!WaveInfo.ReadWaveInfo(FileBytes.GetData(), FileBytes.Num(), &ErrorMessage))
+		{
+			UE_LOG(LogUsd, Warning, TEXT("Failed to read audio info from file '%s': %s"), *FilePath, *ErrorMessage);
+			return nullptr;
+		}
+
+		// If we need to change bit depth, or if the wav internal format is not something we know try converting
+		if (*WaveInfo.pBitsPerSample != 16 || !WaveInfo.IsFormatSupported())
+		{
+			const uint32 OrigNumSamples = Audio::SoundFileUtils::GetNumSamples(FileBytes);
+
+			TArray<uint8> ConvertedRawWaveData;
+
+			// Attempt to convert to 16 bit audio
+			if (Audio::SoundFileUtils::ConvertAudioToWav(FileBytes, ConvertedRawWaveData))
+			{
+				WaveInfo = FWaveModInfo();
+				if (!WaveInfo.ReadWaveInfo(ConvertedRawWaveData.GetData(), ConvertedRawWaveData.Num(), &ErrorMessage))
+				{
+					UE_LOG(LogUsd, Warning, TEXT("Failed to convert audio data from '%s' into a supported wav format: %s"), *FilePath, *ErrorMessage);
+					return nullptr;
+				}
+
+				// Sanity check that the same number of samples exist in the converted file as the original
+				const uint32 ConvertedNumSamples = WaveInfo.GetNumSamples();
+				ensure(ConvertedNumSamples == OrigNumSamples);
+
+				Swap(FileBytes, ConvertedRawWaveData);
+			}
+		}
+
+		if (*WaveInfo.pBitsPerSample != 16)
+		{
+			UE_LOG(LogUsd, Warning, TEXT("Failed to parse audio as .wav data from file '%s' was not 16-bit!"), *FilePath);
+			return nullptr;
+		}
+		check(*WaveInfo.pBitsPerSample == 16);
+
+		int32 ChannelCount = (int32)*WaveInfo.pChannels;
+		check(ChannelCount > 0);
+
+		SoundWave = NewObject<USoundWave>(Outer, AssetName, Flags);
+
+		int32 SizeOfSample = (*WaveInfo.pBitsPerSample) / 8;
+		int32 NumSamples = WaveInfo.SampleDataSize / SizeOfSample;
+		int32 NumFrames = NumSamples / ChannelCount;
+		int32 SamplesPerSec = (int32)*WaveInfo.pSamplesPerSec;
+
+		// Copy the uncompressed data directly into the SoundWave
+		// Note: This currently doesn't work for some reason: We end up with some crash downstream when
+		// creating the audio buffer on the audio playing thread
+		SoundWave->RawPCMDataSize = FileBytes.Num();
+		SoundWave->RawPCMData = (uint8*)FMemory::Malloc(SoundWave->RawPCMDataSize);
+		FMemory::Memcpy(SoundWave->RawPCMData, FileBytes.GetData(), SoundWave->RawPCMDataSize);
+
+		SoundWave->Duration = (float)NumFrames / SamplesPerSec;
+		SoundWave->SetImportedSampleRate(SamplesPerSec);
+		SoundWave->SetSampleRate(SamplesPerSec);
+		SoundWave->NumChannels = ChannelCount;
+		SoundWave->TotalSamples = SamplesPerSec * SoundWave->Duration;
+
+		SoundWave->PostLoad();
+#endif	  // WITH_EDITOR
+
+		return SoundWave;
+	}
+}
+
+void FUsdMediaSpatialAudioTranslator::CreateAssets()
+{
+	TRACE_CPUPROFILER_EVENT_SCOPE(FUsdMediaSpatialAudioTranslator::CreateAssets);
+
+	using namespace UE::UsdMediaSpatialAudioTranslator::Private;
+
+	if (!Context->UsdAssetCache || !Context->InfoCache)
+	{
+		return;
+	}
+
+	// Don't bother generating assets if we're going to just draw some bounds for this prim instead
+	EUsdDrawMode DrawMode = UsdUtils::GetAppliedDrawMode(GetPrim());
+	if (DrawMode != EUsdDrawMode::Default)
+	{
+		CreateAlternativeDrawModeAssets(DrawMode);
+		return;
+	}
+
+	if (!Context->bAllowParsingSounds)
+	{
+		return;
+	}
+
+	FScopedUsdAllocs Allocs;
+
+	const FString PrimPathString = PrimPath.GetString();
+	pxr::UsdPrim UsdPrim = GetPrim();
+	pxr::UsdMediaSpatialAudio UsdAudio{UsdPrim};
+	if (!UsdAudio)
+	{
+		return;
+	}
+	pxr::UsdStageRefPtr Stage = UsdPrim.GetStage();
+
+	pxr::UsdAttribute FilePathAttr = UsdAudio.GetFilePathAttr();
+	FString ResolvedAudioPath = UsdUtils::GetResolvedAssetPath(FilePathAttr, pxr::UsdTimeCode::Default());
+	if (!FPaths::FileExists(ResolvedAudioPath))
+	{
+		UE_LOG(LogUsd, Warning, TEXT("Failed to resolve audio file at path '%s' from prim '%s'"), *ResolvedAudioPath, *PrimPathString);
+		return;
+	}
+
+	FString DesiredName = FPaths::GetBaseFilename(ResolvedAudioPath);
+	FString PrefixedAudioHash = UsdUtils::GetAssetHashPrefix(UsdPrim, Context->bShareAssetsForIdenticalPrims)
+								+ LexToString(FMD5Hash::HashFile(*ResolvedAudioPath));
+
+	USoundWave* SoundWave = Context->UsdAssetCache->GetOrCreateCustomCachedAsset<USoundWave>(
+		PrefixedAudioHash,
+		DesiredName,
+		Context->ObjectFlags,
+		[&ResolvedAudioPath](UPackage* Outer, FName SanitizedName, EObjectFlags FlagsToUse) -> UObject*
+		{
+			return UE::UsdMediaSpatialAudioTranslator::Private::CreateSoundWave(	//
+				ResolvedAudioPath,
+				SanitizedName,
+				Outer,
+				FlagsToUse
+			);
+		}
+	);
+
+	if (SoundWave)
+	{
+		Context->InfoCache->LinkAssetToPrim(PrimPath, SoundWave);
+
+		if (UUsdAssetUserData* UserData = UsdUnreal::ObjectUtils::GetOrCreateAssetUserData<UUsdAssetUserData>(SoundWave))
+		{
+			UserData->PrimPaths.AddUnique(PrimPathString);
+
+			if (Context->MetadataOptions.bCollectMetadata)
+			{
+				UsdToUnreal::ConvertMetadata(
+					UsdPrim,
+					UserData,
+					Context->MetadataOptions.BlockedPrefixFilters,
+					Context->MetadataOptions.bInvertFilters,
+					Context->MetadataOptions.bCollectFromEntireSubtrees
+				);
+			}
+			else
+			{
+				UserData->StageIdentifierToMetadata.Remove(UsdToUnreal::ConvertString(Stage->GetRootLayer()->GetIdentifier()));
+			}
+		}
+	}
+}
+
+USceneComponent* FUsdMediaSpatialAudioTranslator::CreateComponents()
+{
+	TRACE_CPUPROFILER_EVENT_SCOPE(FUsdMediaSpatialAudioTranslator::CreateComponents);
+
+	USceneComponent* SceneComponent = nullptr;
+
+	EUsdDrawMode DrawMode = UsdUtils::GetAppliedDrawMode(GetPrim());
+	if (DrawMode == EUsdDrawMode::Default)
+	{
+		if (Context->bAllowParsingSounds)
+		{
+			const bool bNeedsActor = true;
+			SceneComponent = CreateComponentsEx({UAudioComponent::StaticClass()}, bNeedsActor);
+		}
+	}
+	else
+	{
+		SceneComponent = CreateAlternativeDrawModeComponents(DrawMode);
+	}
+
+	UpdateComponents(SceneComponent);
+
+	return SceneComponent;
+}
+
+void FUsdMediaSpatialAudioTranslator::UpdateComponents(USceneComponent* SceneComponent)
+{
+	// Note how we don't even set the audio on the component here: All of our audio playback is done via
+	// the Sequencer. Essentially the audio actor/component are exclusively used for their transform, in case
+	// the Sequencer is playing spatial audio.
+	//
+	// We exclusively use the Sequencer for playback and properties for a few reasons:
+	//  - We must use the Sequencer for audio in the first place anyway, as there is no easy way of scrubbing/restarting
+	//    audio playback directly via the audio component or Time attribute animation. As I understand it, the audio component
+	//    is intentionally very "simple", and is meant to be driven by blueprint or just play its audio once at spawn and be
+	//    disabled (e.g. for an explosion sound effect). It's not meant to be a full "media player";
+	//  - It's much easier to manipulate the attributes like startTime/endTime, mediaOffset, volume and looping via
+	//    the Sequencer from the user standpoint, as they're all right on the section and you can just click and drag keyframes,
+	//    section boundaries, the entire section, etc.;
+	//  - Putting all the attributes on the Sequencer section means we can reuse the same audio asset for different
+	//    attribute configurations. For example, if we wanted to play the audio directly via the component, we'd have to store whether
+	//    it is looping or not within the asset itself. This means that playing two versions of the audio simultaneously,
+	//    one looping and one not, would have needed two separate copies of the same audio asset...
+	//  - If we had done all of the above and *also* placed the audio and whatever properties we could on the component, it would
+	//    have just been more confusing, as a user wouldn't really know what was actually driving the audio or not
+
+	Super::UpdateComponents(SceneComponent);
+}
+
+bool FUsdMediaSpatialAudioTranslator::CollapsesChildren(ECollapsingType CollapsingType) const
+{
+	// If we have a custom draw mode, it means we should draw bounds/cards/etc. instead
+	// of our entire subtree, which is basically the same thing as collapsing
+	EUsdDrawMode DrawMode = UsdUtils::GetAppliedDrawMode(GetPrim());
+	if (DrawMode != EUsdDrawMode::Default)
+	{
+		return true;
+	}
+
+	return false;
+}
+
+bool FUsdMediaSpatialAudioTranslator::CanBeCollapsed(ECollapsingType CollapsingType) const
+{
+	return false;
+}
+
+TSet<UE::FSdfPath> FUsdMediaSpatialAudioTranslator::CollectAuxiliaryPrims() const
+{
+	return {};
+}
+
+#endif	  // #if USE_USD_SDK
diff --git a/Engine/Plugins/Importers/USDImporter/Source/USDSchemas/Private/USDSchemasModule.cpp b/Engine/Plugins/Importers/USDImporter/Source/USDSchemas/Private/USDSchemasModule.cpp
--- a/Engine/Plugins/Importers/USDImporter/Source/USDSchemas/Private/USDSchemasModule.cpp
+++ b/Engine/Plugins/Importers/USDImporter/Source/USDSchemas/Private/USDSchemasModule.cpp
@@ -23,6 +23,7 @@
 #include "USDGeomXformableTranslator.h"
 #include "USDGroomTranslator.h"
 #include "USDLuxLightTranslator.h"
+#include "USDMediaSpatialAudioTranslator.h"
 #include "USDMemory.h"
 #include "USDShadeMaterialTranslator.h"
 #include "USDSkelSkeletonTranslator.h"
@@ -76,6 +77,8 @@ public:
 				Registry.Register<FUsdGroomTranslator>(TEXT("UsdGeomXformable")),
 				// The GeometryCacheTranslator also works on UsdGeomXformable through the GroomTranslator
 				Registry.Register<FUsdGeometryCacheTranslator>(TEXT("UsdGeomMesh")),
+				// It doesn't seem possible to create SoundWave assets at runtime at the moment, for whatever reason
+				Registry.Register<FUsdMediaSpatialAudioTranslator>(TEXT("UsdMediaSpatialAudio"))
 			});
 
 			if (IMDLImporterModule* MDLImporterModule = FModuleManager::Get().LoadModulePtr<IMDLImporterModule>(TEXT("MDLImporter")))
diff --git a/Engine/Plugins/Importers/USDImporter/Source/USDSchemas/Public/USDMediaSpatialAudioTranslator.h b/Engine/Plugins/Importers/USDImporter/Source/USDSchemas/Public/USDMediaSpatialAudioTranslator.h
new file mode 100644
--- /dev/null
+++ b/Engine/Plugins/Importers/USDImporter/Source/USDSchemas/Public/USDMediaSpatialAudioTranslator.h
@@ -0,0 +1,29 @@
+// Copyright Epic Games, Inc. All Rights Reserved.
+
+#pragma once
+
+#include "USDGeomXformableTranslator.h"
+
+#if USE_USD_SDK
+
+class USceneComponent;
+
+class USDSCHEMAS_API FUsdMediaSpatialAudioTranslator : public FUsdGeomXformableTranslator
+{
+	using Super = FUsdGeomXformableTranslator;
+
+public:
+
+	using FUsdGeomXformableTranslator::FUsdGeomXformableTranslator;
+
+	virtual void CreateAssets() override;
+	virtual USceneComponent* CreateComponents() override;
+	virtual void UpdateComponents(USceneComponent* SceneComponent) override;
+
+	virtual bool CollapsesChildren(ECollapsingType CollapsingType) const override;
+	virtual bool CanBeCollapsed(ECollapsingType CollapsingType) const override;
+
+	virtual TSet<UE::FSdfPath> CollectAuxiliaryPrims() const override;
+};
+
+#endif	  // #if USE_USD_SDK
diff --git a/Engine/Plugins/Importers/USDImporter/Source/USDSchemas/Public/USDSchemaTranslator.h b/Engine/Plugins/Importers/USDImporter/Source/USDSchemas/Public/USDSchemaTranslator.h
--- a/Engine/Plugins/Importers/USDImporter/Source/USDSchemas/Public/USDSchemaTranslator.h
+++ b/Engine/Plugins/Importers/USDImporter/Source/USDSchemas/Public/USDSchemaTranslator.h
@@ -305,6 +305,9 @@ struct USDSCHEMAS_API FUsdSchemaTranslationContext : public TSharedFromThis<FUsd
 	/** If true, means we will try generating Sparse Volume Textures */
 	bool bAllowParsingSparseVolumeTextures = true;
 
+	/** If true, means we will try generating SoundWave assets from sound files referenced by UsdMediaSpatialAudio prims */
+	bool bAllowParsingSounds = true;
+
 	/** Skip the import of materials that aren't being used by any prim on the stage */
 	bool bTranslateOnlyUsedMaterials = false;
 
diff --git a/Engine/Plugins/Importers/USDImporter/Source/USDSchemas/USDSchemas.Build.cs b/Engine/Plugins/Importers/USDImporter/Source/USDSchemas/USDSchemas.Build.cs
--- a/Engine/Plugins/Importers/USDImporter/Source/USDSchemas/USDSchemas.Build.cs
+++ b/Engine/Plugins/Importers/USDImporter/Source/USDSchemas/USDSchemas.Build.cs
@@ -53,6 +53,7 @@ namespace UnrealBuildTool.Rules
 				PrivateDependencyModuleNames.AddRange(
 					new string[]
 					{
+						"AudioEditor", // For USoundFactory, which we use for parsing UsdMediaSpatialAudio
 						"BlueprintGraph",
 						"GeometryCacheUSD",
 						"HairStrandsEditor",
diff --git a/Engine/Plugins/Importers/USDImporter/Source/USDStage/Private/USDLevelSequenceHelper.cpp b/Engine/Plugins/Importers/USDImporter/Source/USDStage/Private/USDLevelSequenceHelper.cpp
--- a/Engine/Plugins/Importers/USDImporter/Source/USDStage/Private/USDLevelSequenceHelper.cpp
+++ b/Engine/Plugins/Importers/USDImporter/Source/USDStage/Private/USDLevelSequenceHelper.cpp
@@ -65,13 +65,17 @@
 #include "MovieSceneTrack.h"
 #include "Rigs/FKControlRig.h"
 #include "RigVMBlueprintGeneratedClass.h"
+#include "Sections/MovieSceneAudioSection.h"
 #include "Sections/MovieSceneFloatSection.h"
 #include "Sections/MovieSceneSubSection.h"
 #include "Sequencer/MovieSceneControlRigParameterSection.h"
 #include "Sequencer/MovieSceneControlRigParameterTrack.h"
+#include "Sound/SoundAttenuation.h"
+#include "Sound/SoundBase.h"
 #include "SparseVolumeTexture/SparseVolumeTexture.h"
 #include "Templates/SharedPointer.h"
 #include "Tracks/MovieScene3DTransformTrack.h"
+#include "Tracks/MovieSceneAudioTrack.h"
 #include "Tracks/MovieSceneBoolTrack.h"
 #include "Tracks/MovieSceneColorTrack.h"
 #include "Tracks/MovieSceneFloatTrack.h"
@@ -95,6 +99,12 @@
 #include "Subsystems/AssetEditorSubsystem.h"
 #endif	  // WITH_EDITOR
 
+#if USE_USD_SDK
+#include "USDIncludesStart.h"
+#include "pxr/usd/usdMedia/tokens.h"
+#include "USDIncludesEnd.h"
+#endif	  // USE_USD_SDK
+
 #define LOCTEXT_NAMESPACE "USDLevelSequenceHelper"
 
 #if USE_USD_SDK
@@ -698,6 +708,7 @@ private:
 	void AddGroomTracks(const UUsdPrimTwin& PrimTwin, const UE::FUsdPrim& Prim);
 	void AddGeometryCacheTracks(const UUsdPrimTwin& PrimTwin, const UE::FUsdPrim& Prim);
 	void AddVolumeTracks(const UUsdPrimTwin& PrimTwin, const UE::FUsdPrim& Prim);
+	void AddAudioTracks(const UUsdPrimTwin& PrimTwin, const UE::FUsdPrim& Prim);
 
 	template<typename TrackType>
 	TrackType* AddTrack(
@@ -2663,6 +2674,315 @@ void FUsdLevelSequenceHelperImpl::AddVolumeTracks(const UUsdPrimTwin& PrimTwin,
 	PrimPathByLevelSequenceName.AddUnique(PrimSequence->GetFName(), Prim.GetPrimPath().GetString());
 }
 
+namespace UE::LevelSequenceHelper::Private
+{
+	template<typename T>
+	TOptional<T> GetAuthoredValue(const UE::FUsdPrim& Prim, const FString& AttrName)
+	{
+		if (UE::FUsdAttribute Attr = Prim.GetAttribute(*AttrName))
+		{
+			UE::FVtValue VtValue;
+			if (Attr.HasAuthoredValue() && Attr.Get(VtValue) && !VtValue.IsEmpty())
+			{
+				return UsdUtils::GetUnderlyingValue<T>(VtValue);
+			}
+		}
+
+		return TOptional<T>{};
+	}
+};
+
+void FUsdLevelSequenceHelperImpl::AddAudioTracks(const UUsdPrimTwin& PrimTwin, const UE::FUsdPrim& Prim)
+{
+	using namespace UnrealIdentifiers;
+	using namespace UE::LevelSequenceHelper::Private;
+
+	UAudioComponent* AudioComponent = Cast<UAudioComponent>(PrimTwin.GetSceneComponent());
+	if (!AudioComponent)
+	{
+		return;
+	}
+
+	// Note: We pull the audio directly from the info cache here, and not the component:
+	// See big comment within FUsdMediaSpatialAudioTranslator::UpdateComponents
+	USoundBase* Sound = InfoCache->GetSingleAssetForPrim<USoundBase>(Prim.GetPrimPath());
+	if (!Sound)
+	{
+		return;
+	}
+
+	static FString FilePathToken = UsdToUnreal::ConvertToken(pxr::UsdMediaTokens->filePath);
+	static FString AuralModeToken = UsdToUnreal::ConvertToken(pxr::UsdMediaTokens->auralMode);
+	static FString MediaOffsetSecondsToken = UsdToUnreal::ConvertToken(pxr::UsdMediaTokens->mediaOffset);
+	static FString StartTimeCodeToken = UsdToUnreal::ConvertToken(pxr::UsdMediaTokens->startTime);
+	static FString EndTimeCodeToken = UsdToUnreal::ConvertToken(pxr::UsdMediaTokens->endTime);
+	static FString PlaybackModeToken = UsdToUnreal::ConvertToken(pxr::UsdMediaTokens->playbackMode);
+	static FString GainToken = UsdToUnreal::ConvertToken(pxr::UsdMediaTokens->gain);
+
+	// Set up the actual audio track
+	{
+		// If we're using references so that our prim references another layer with the SpatialAudio prim,
+		// using just "the layer for the prim" here would mean we end up with the referencer layer, and so
+		// we'd end up placing the audio track directly on the LevelSequence for the referencer layer, which
+		// is not what we want.
+		// Instead, we use the file path attribute as the "main attribute" for the audio track: If you author
+		// filePath on the referencer layer, the audio track will end up on the corresponding LevelSequence.
+		// If you author it on the referenced layer, the audio track will end up on that subsequence
+		UE::FUsdAttribute Attr = Prim.GetAttribute(*FilePathToken);
+
+		UE::FSdfLayer SequenceLayer;
+		ULevelSequence* AttributeSequence = FindOrAddSequenceForAttribute(Attr, &SequenceLayer);
+		if (!AttributeSequence)
+		{
+			return;
+		}
+
+		UMovieScene* MovieScene = AttributeSequence->GetMovieScene();
+		if (!MovieScene)
+		{
+			return;
+		}
+
+		UsdLevelSequenceHelperImpl::FScopedReadOnlyDisable ReadOnlyGuard{MovieScene, SequenceLayer, UsdStage};
+
+		// Since the schema has fallbacks for these, we should always have values here.
+		//
+		// Note that since startTime and endTime have pxr::SdfTimeCode data types, USD will already do the
+		// proper conversions regarding sublayer/reference offset and scale. In other words, these timeCode
+		// values as retrieved by UE::FUsdAttribute::Get() are *already relative to the stage*.
+		// Example: If startTime was set to 15 for a SpatialAudio prim defined in a sublayer, added to the stage
+		// with offset of 10 and scale of 2, when querying the attribute for the prim we'd get that startTime
+		// actually has the value of 40, because 15 * 2 + 10 = 40
+		TOptional<double> MediaOffset = GetAuthoredValue<double>(Prim, MediaOffsetSecondsToken);
+		TOptional<pxr::SdfTimeCode> StartTimeCode = GetAuthoredValue<pxr::SdfTimeCode>(Prim, StartTimeCodeToken);
+		TOptional<pxr::SdfTimeCode> EndTimeCode = GetAuthoredValue<pxr::SdfTimeCode>(Prim, EndTimeCodeToken);
+		TOptional<pxr::TfToken> PlaybackMode = GetAuthoredValue<pxr::TfToken>(Prim, PlaybackModeToken);
+		TOptional<pxr::TfToken> AuralMode = GetAuthoredValue<pxr::TfToken>(Prim, AuralModeToken);
+
+		// The documentation mentions to swap these in edge cases like negative scaling
+		if (StartTimeCode.IsSet() && EndTimeCode.IsSet())
+		{
+			if (StartTimeCode.GetValue() > EndTimeCode.GetValue())
+			{
+				Swap(StartTimeCode, EndTimeCode);
+			}
+		}
+
+		// Handle the different playback modes
+		bool bIsLooping = false;
+		if (PlaybackMode.IsSet())
+		{
+			const pxr::TfToken PlaybackModeValue = PlaybackMode.GetValue();
+			if (PlaybackModeValue == pxr::UsdMediaTokens->onceFromStart)
+			{
+				// "Play the audio once, starting at startTime, continuing until the audio completes."
+				EndTimeCode.Reset();
+			}
+			else if (PlaybackModeValue == pxr::UsdMediaTokens->onceFromStartToEnd)
+			{
+				// "Play the audio once beginning at startTime, continuing until endTime or until the
+				// audio completes, whichever comes first."
+				//
+				// Do nothing: Just continue using startTime and endTime as we have already
+			}
+			else if (PlaybackModeValue == pxr::UsdMediaTokens->loopFromStart)
+			{
+				// "Start playing the audio at startTime and continue looping through to the stage's
+				// authored endTimeCode."
+				bIsLooping = true;
+				EndTimeCode = Prim.GetStage().GetEndTimeCode();
+			}
+			else if (PlaybackModeValue == pxr::UsdMediaTokens->loopFromStartToEnd)
+			{
+				// "Start playing the audio at startTime and continue looping through, stopping
+				// the audio at endTime."
+				bIsLooping = true;
+			}
+			else if (PlaybackModeValue == pxr::UsdMediaTokens->loopFromStage)
+			{
+				// "Start playing the audio at the stage's authored startTimeCode and continue looping
+				// through to the stage's authored endTimeCode. This can be useful for ambient sounds
+				// that should always be active."
+				bIsLooping = true;
+
+				// Fetch time range from the stage instead
+				UE::FSdfLayer RootLayer = Prim.GetStage().GetRootLayer();
+				if (RootLayer)
+				{
+					StartTimeCode = RootLayer.GetStartTimeCode();
+
+					// Since the fallback value for EndTimeCode is also 0, we have to take care to check
+					// whether the stage actually has it authored or not
+					if (RootLayer.HasEndTimeCode())
+					{
+						EndTimeCode = RootLayer.GetEndTimeCode();
+					}
+					else
+					{
+						EndTimeCode.Reset();
+					}
+				}
+			}
+		}
+
+		// Convert time codes to being relative to the stage to being relative to the Subsequence they are in
+		//
+		// This may seem like it undoes the conversions above, and it really does: For cases where we have a
+		// sublayer with an offset and scale, we'll end up adding the layer-local time samples to the subsequence,
+		// like we want. Using PrimToStage AND the SequenceTransform is needed for a different case however: Prim
+		// references with sublayer and offsets. In that case the prim itself may have a sublayer and offset, but
+		// its track will be placed in the levelsequence for the *referencer* layer: This means we want to see the
+		// keys on that layer instead, at times relative to it
+		FFrameNumber StartFrameNumber;
+		TOptional<FFrameNumber> EndFrameNumber;
+		FFrameNumber MediaFrameNumberOffset;
+		{
+			FMovieSceneSequenceTransform SequenceTransform;
+			FMovieSceneSequenceID SequenceID = SequencesID.FindRef(AttributeSequence);
+			if (FMovieSceneSubSequenceData* SubSequenceData = SequenceHierarchyCache.FindSubData(SequenceID))
+			{
+				SequenceTransform = SubSequenceData->RootToSequenceTransform;
+			}
+
+			const FFrameRate Resolution = MovieScene->GetTickResolution();
+			const FFrameRate DisplayRate = MovieScene->GetDisplayRate();
+			const double StageTimeCodesPerSecond = Prim.GetStage().GetTimeCodesPerSecond();
+			const FFrameRate StageFrameRate(StageTimeCodesPerSecond, 1);
+
+			if (StartTimeCode.IsSet())
+			{
+				double DoubleValue = static_cast<double>(StartTimeCode.GetValue());
+				int32 FrameNumber = FMath::FloorToInt(DoubleValue);
+				float SubFrameNumber = DoubleValue - FrameNumber;
+				FFrameTime FrameTime(FrameNumber, SubFrameNumber);
+				FFrameTime KeyFrameTime = FFrameRate::TransformTime(FrameTime, StageFrameRate, Resolution);
+				KeyFrameTime *= SequenceTransform;
+				StartFrameNumber = KeyFrameTime.FloorToFrame();
+			}
+
+			if (EndTimeCode.IsSet())
+			{
+				double DoubleValue = static_cast<double>(EndTimeCode.GetValue());
+				int32 FrameNumber = FMath::FloorToInt(DoubleValue);
+				float SubFrameNumber = DoubleValue - FrameNumber;
+				FFrameTime FrameTime(FrameNumber, SubFrameNumber);
+				FFrameTime KeyFrameTime = FFrameRate::TransformTime(FrameTime, StageFrameRate, Resolution);
+				KeyFrameTime *= SequenceTransform;
+				EndFrameNumber = KeyFrameTime.FloorToFrame();
+			}
+
+			if (MediaOffset.IsSet())
+			{
+				double OffsetSeconds = MediaOffset.GetValue();
+				MediaFrameNumberOffset = Resolution.AsFrameNumber(OffsetSeconds);
+			}
+		}
+
+		// Get the attenuation settings we'll use if we're in spatial mode
+		USoundAttenuation* Attenuation = nullptr;
+		if (AuralMode.IsSet() && AuralMode.GetValue() == pxr::UsdMediaTokens->spatial)
+		{
+			const UUsdProjectSettings* ProjectSettings = GetDefault<UUsdProjectSettings>();
+			if (ProjectSettings)
+			{
+				Attenuation = Cast<USoundAttenuation>(ProjectSettings->DefaultSoundAttenuation.TryLoad());
+			}
+		}
+
+		MovieScene->Modify();
+		MovieScene->SetClockSource(EUpdateClockSource::Audio);
+
+		// We name it AudioTrack here instead of prim name because in general we'll already have the actor and component
+		// binding show the prim name anyway... It's not very useful to see "MyPrim / MyPrim / MyPrim" on the Sequencer
+		const bool bIsMuted = false;
+		if (UMovieSceneAudioTrack* AudioTrack = AddTrack<
+				UMovieSceneAudioTrack>(TEXT("Audio track"), PrimTwin, *AudioComponent, *AttributeSequence, bIsMuted))
+		{
+			AudioTrack->RemoveAllAnimationData();
+
+			const FFrameNumber Time = StartFrameNumber;
+			UMovieSceneAudioSection* NewAudioSection = Cast<UMovieSceneAudioSection>(AudioTrack->AddNewSound(Sound, Time));
+
+			if (NewAudioSection)
+			{
+				NewAudioSection->Modify();
+				NewAudioSection->SetLooping(bIsLooping);
+				NewAudioSection->SetStartOffset(MediaFrameNumberOffset);
+
+				// Parse gain into channel volume
+				UE::FUsdAttribute GainAttr = Prim.GetAttribute(*GainToken);
+				if (GainAttr && GainAttr.HasAuthoredValue())
+				{
+					// We have to dig through the channel proxy because UMovieSceneAudioSection doesn't expose non-const access
+					// to the volume channel...
+					FMovieSceneChannelProxy& ChannelProxy = NewAudioSection->GetChannelProxy();
+					FMovieSceneFloatChannel* VolumeChannel = ChannelProxy.GetChannel<FMovieSceneFloatChannel>(0);
+					if (VolumeChannel)
+					{
+						// Set default value into the channel in case we don't have any animation
+						UE::FVtValue DefaultValue;
+						if (GainAttr.Get(DefaultValue))
+						{
+							TOptional<double> DefaultGain = UsdUtils::GetUnderlyingValue<double>(DefaultValue);
+							if (DefaultGain.IsSet())
+							{
+								VolumeChannel->SetDefault(static_cast<float>(DefaultGain.GetValue()));
+							}
+						}
+
+						FMovieSceneSequenceTransform SequenceTransform;
+						FMovieSceneSequenceID SequenceID = SequencesID.FindRef(AttributeSequence);
+						if (FMovieSceneSubSequenceData* SubSequenceData = SequenceHierarchyCache.FindSubData(SequenceID))
+						{
+							SequenceTransform = SubSequenceData->RootToSequenceTransform;
+						}
+
+						// Parse time-sampled animation into the volume channel
+						TArray<double> TimeSamples;
+						if (GainAttr.GetTimeSamples(TimeSamples) && TimeSamples.Num() > 0)
+						{
+							UsdToUnreal::FPropertyTrackReader Reader = UsdToUnreal::CreatePropertyTrackReader(Prim, TEXT("Volume"));
+							UsdToUnreal::ConvertFloatTimeSamples(
+								UsdStage,
+								TimeSamples,
+								Reader.FloatReader,
+								*VolumeChannel,
+								*MovieScene,
+								SequenceTransform
+							);
+						}
+					}
+				}
+
+				// Enable spatial audio if we should
+				if (Attenuation)
+				{
+					const bool bOverrideAttenuation = true;
+					NewAudioSection->SetOverrideAttenuation(bOverrideAttenuation);
+					NewAudioSection->SetAttenuationSettings(Attenuation);
+				}
+
+				// Start with the auto-size range because the section can't be unbounded, so if we want
+				// the audio to play to completion we'd otherwise need to specify the end of the section
+				// ourselves in order to get a valid range
+				TOptional<TRange<FFrameNumber>> Range = NewAudioSection->GetAutoSizeRange();
+				if (Range.IsSet())
+				{
+					TRange<FFrameNumber>& RangeValue = Range.GetValue();
+					RangeValue.SetLowerBound(TRangeBound<FFrameNumber>{StartFrameNumber});
+					if (EndFrameNumber.IsSet())
+					{
+						RangeValue.SetUpperBound(TRangeBound<FFrameNumber>{EndFrameNumber.GetValue()});
+					}
+					NewAudioSection->SetRange(RangeValue);
+				}
+			}
+		}
+
+		PrimPathByLevelSequenceName.AddUnique(AttributeSequence->GetFName(), Prim.GetPrimPath().GetString());
+	}
+}
+
 void FUsdLevelSequenceHelperImpl::AddPrim(UUsdPrimTwin& PrimTwin, bool bForceVisibilityTracks, TOptional<bool> HasAnimatedBounds)
 {
 	if (!UsdStage)
@@ -2731,6 +3051,10 @@ void FUsdLevelSequenceHelperImpl::AddPrim(UUsdPrimTwin& PrimTwin, bool bForceVis
 	{
 		AddVolumeTracks(PrimTwin, UsdPrim);
 	}
+	else if (UsdPrim.IsA(TEXT("SpatialAudio")))
+	{
+		AddAudioTracks(PrimTwin, UsdPrim);
+	}
 
 	AddCommonTracks(PrimTwin, UsdPrim, bForceVisibilityTracks);
 
@@ -2783,6 +3107,10 @@ TrackType* FUsdLevelSequenceHelperImpl::AddTrack(
 		{
 			Track->SetDisplayName(FText::FromName(TrackName));
 		}
+		else if constexpr (std::is_base_of_v<UMovieSceneAudioTrack, TrackType>)
+		{
+			Track->SetDisplayName(FText::FromName(TrackName));
+		}
 #endif	  // WITH_EDITOR
 	}
 
@@ -4247,6 +4575,37 @@ void FUsdLevelSequenceHelperImpl::HandleTrackChange(const UMovieSceneTrack& Trac
 					RefreshSequencer();
 				}
 			}
+			else if (const UMovieSceneAudioTrack* AudioTrack = Cast<const UMovieSceneAudioTrack>(&Track))
+			{
+				const TArray<UMovieSceneSection*>& Sections = AudioTrack->GetAudioSections();
+				if (Sections.Num() > 1)
+				{
+					UE_LOG(
+						LogUsd,
+						Warning,
+						TEXT(
+							"The audio track '%s' has %d sections, but only the first audio section of an audio track can be written out to USD for now"
+						),
+						*AudioTrack->GetPathName(),
+						Sections.Num()
+					);
+				}
+
+				if (Sections.Num() > 0)
+				{
+					if (UMovieSceneAudioSection* AudioSection = Cast<UMovieSceneAudioSection>(Sections[0]))
+					{
+						FMovieSceneSequenceTransform Identity;
+						UnrealToUsd::ConvertAudioSection(*AudioSection, Identity, UsdPrim);
+					}
+				}
+			}
+		}
+
+		// Notify the USD Stage Editor to refresh this prim the next frame
+		if (AUsdStageActor* StageActorPtr = StageActor.Get())
+		{
+			StageActorPtr->OnPrimChanged.Broadcast(PrimTwin->PrimPath, false);
 		}
 	}
 }
diff --git a/Engine/Plugins/Importers/USDImporter/Source/USDStage/Private/USDStageActor.cpp b/Engine/Plugins/Importers/USDImporter/Source/USDStage/Private/USDStageActor.cpp
--- a/Engine/Plugins/Importers/USDImporter/Source/USDStage/Private/USDStageActor.cpp
+++ b/Engine/Plugins/Importers/USDImporter/Source/USDStage/Private/USDStageActor.cpp
@@ -41,6 +41,7 @@
 #include "Async/ParallelFor.h"
 #include "CineCameraActor.h"
 #include "CineCameraComponent.h"
+#include "Components/AudioComponent.h"
 #include "Components/DirectionalLightComponent.h"
 #include "Components/HierarchicalInstancedStaticMeshComponent.h"
 #include "Components/LightComponent.h"
@@ -88,15 +89,17 @@
 #include "MovieSceneDynamicBindingUtils.h"
 #include "ScopedTransaction.h"
 #include "Selection.h"
+#include "SequencerUtilities.h"
 #include "Subsystems/AssetEditorSubsystem.h"
 #include "UnrealEdGlobals.h"
 #include "USDClassesEditorModule.h"
-#include "SequencerUtilities.h"
+
 #endif	  // WITH_EDITOR
 
 #if USE_USD_SDK
 #include "USDIncludesStart.h"
 #include "pxr/usd/usdGeom/tokens.h"
+#include "pxr/usd/usdMedia/tokens.h"
 #include "pxr/usd/usdPhysics/tokens.h"
 #include "USDIncludesEnd.h"
 #endif	  // USE_USD_SDK
@@ -618,7 +621,12 @@ struct FUsdStageActorImpl
 			return;
 		}
 
-		UMovieSceneReplaceableDirectorBlueprintBinding* NewCustomBinding = Cast<UMovieSceneReplaceableDirectorBlueprintBinding>(UMovieSceneReplaceableDirectorBlueprintBinding::StaticClass()->GetDefaultObject<UMovieSceneCustomBinding>()->CreateNewCustomBinding(nullptr, *MovieScene));
+		UMovieSceneReplaceableDirectorBlueprintBinding* NewCustomBinding = Cast<UMovieSceneReplaceableDirectorBlueprintBinding>(
+			UMovieSceneReplaceableDirectorBlueprintBinding::StaticClass()->GetDefaultObject<UMovieSceneCustomBinding>()->CreateNewCustomBinding(
+				nullptr,
+				*MovieScene
+			)
+		);
 
 		if (!NewCustomBinding)
 		{
@@ -636,19 +644,16 @@ struct FUsdStageActorImpl
 			);
 			if (DynamicBinding.Function)
 			{
-				DynamicBinding.ResolveParamsProperty = DynamicBinding.Function->FindPropertyByName(
-					DynamicBinding.ResolveParamsPinName
-				);
+				DynamicBinding.ResolveParamsProperty = DynamicBinding.Function->FindPropertyByName(DynamicBinding.ResolveParamsPinName);
 			}
 
 			// Store a path to this very actor on the binding, so that it can find us later and ask how to resolve a particular
 			// prim path
 			FMovieSceneDynamicBindingPayloadVariable& ActorPathVariable = DynamicBinding.PayloadVariables.FindOrAdd(TEXT("StageActorIDNam"
-				"eFilter"));
+																														 "eFilter"));
 			ActorPathVariable.Value = DefaultActorFilter;
 
-			FMovieSceneDynamicBindingPayloadVariable& RootLayerVariable = DynamicBinding.PayloadVariables.FindOrAdd(TEXT("RootLayerFilter"
-			));
+			FMovieSceneDynamicBindingPayloadVariable& RootLayerVariable = DynamicBinding.PayloadVariables.FindOrAdd(TEXT("RootLayerFilter"));
 			RootLayerVariable.Value = FString{};	// No root layer filter by default for more flexibility
 
 			FMovieSceneDynamicBindingPayloadVariable& PrimPathVariable = DynamicBinding.PayloadVariables.FindOrAdd(TEXT("PrimPath"));
@@ -1627,6 +1632,23 @@ void AUsdStageActor::OnUsdObjectsChanged(const UsdUtils::FObjectChangesByPath& I
 					bHasResync = true;
 					break;
 				}
+
+				// We put all of the audio info directly on the section, so if any of these change then we need to re-add the prim
+				// to the sequencer, and potentially generate a new audio asset
+				static const TSet<FString> AudioSectionProperties = {
+					*UsdToUnreal::ConvertToken(pxr::UsdMediaTokens->filePath),
+					*UsdToUnreal::ConvertToken(pxr::UsdMediaTokens->auralMode),
+					*UsdToUnreal::ConvertToken(pxr::UsdMediaTokens->playbackMode),
+					*UsdToUnreal::ConvertToken(pxr::UsdMediaTokens->startTime),
+					*UsdToUnreal::ConvertToken(pxr::UsdMediaTokens->endTime),
+					*UsdToUnreal::ConvertToken(pxr::UsdMediaTokens->mediaOffset),
+					*UsdToUnreal::ConvertToken(pxr::UsdMediaTokens->gain)};
+				if (AudioSectionProperties.Contains(AttributeChange.PropertyName))
+				{
+					bIsResync = true;
+					bHasResync = true;
+					break;
+				}
 			}
 
 			// Any sublayer change (even offsets) means we need to regenerate our LevelSequence to add (or shift)
@@ -1683,6 +1705,10 @@ void AUsdStageActor::OnUsdObjectsChanged(const UsdUtils::FObjectChangesByPath& I
 		// to update certain info cache maps: For example whenever we delete a prim we need to make sure it's removed from MaterialUsers, etc.
 		TSharedRef<FUsdSchemaTranslationContext> TranslationContext = FUsdStageActorImpl::CreateUsdSchemaTranslationContext(this, TEXT("/"));
 		InfoCache->RebuildCacheForSubtree(Stage.GetPseudoRoot(), TranslationContext.Get());
+
+		// Need to update the LevelSequenceHelper with our new cache too, as it will need the new cache to find any
+		// assets we end up creating in this update (e.g. USoundWave assets)
+		LevelSequenceHelper.SetInfoCache(InfoCache);
 	}
 
 	if (BBoxCache.IsValid() && bHasResync)
@@ -2267,7 +2293,13 @@ UUsdPrimTwin* AUsdStageActor::ExpandPrim(const UE::FUsdPrim& Prim, bool bResync,
 	}
 	if (!bIsAnimated)
 	{
-		bIsAnimated = UsdUtils::IsAnimated(Prim);
+		// Always consider SpatialAudio prims as animated so that we can create LevelSequence tracks for the audio itself.
+		// We exclusively handle the audio stuff via the Sequencer and LevelSequence tracks because there's no way to play audio via Time animation,
+		// and the audio component is not meant to be a fully featured audio player with start/end play times and animated volume controls. In
+		// other words, if we placed our SoundWave asset on the component, the audio component would instantly play it when going into PIE, which is
+		// not what we want. The audio component and actor are only really used for their transforms on the level whenever we'ret trying to play
+		// spatial audio
+		bIsAnimated = UsdUtils::IsAnimated(Prim) || Prim.IsA(TEXT("SpatialAudio"));
 	}
 
 	// Create Sequencer tracks for the prim
@@ -4369,23 +4401,31 @@ void AUsdStageActor::OnMovieSceneDataChanged(EMovieSceneDataChangeType ChangeTyp
 
 		// If the binding already has another dynamic binding let's not touch it,
 		// regardless of whether we set that dynamic binding up or the user did
-		if (Algo::AnyOf(BindingReferences->GetReferences(Possessable.GetGuid()), [](const FMovieSceneBindingReference& Reference) {
-			if (UMovieSceneSpawnableDirectorBlueprintBinding* SpawnableBinding = Cast<UMovieSceneSpawnableDirectorBlueprintBinding>(Reference.CustomBinding))
-			{
-				if (SpawnableBinding->DynamicBinding.Function != nullptr)
-				{
-					return true;
-				}
-			}
-			if (UMovieSceneReplaceableDirectorBlueprintBinding* ReplaceableBinding = Cast<UMovieSceneReplaceableDirectorBlueprintBinding>(Reference.CustomBinding))
-			{
-				if (ReplaceableBinding->DynamicBinding.Function != nullptr)
+		if (Algo::AnyOf(
+				BindingReferences->GetReferences(Possessable.GetGuid()),
+				[](const FMovieSceneBindingReference& Reference)
 				{
-					return true;
+					if (UMovieSceneSpawnableDirectorBlueprintBinding* SpawnableBinding = Cast<UMovieSceneSpawnableDirectorBlueprintBinding>(
+							Reference.CustomBinding
+						))
+					{
+						if (SpawnableBinding->DynamicBinding.Function != nullptr)
+						{
+							return true;
+						}
+					}
+					if (UMovieSceneReplaceableDirectorBlueprintBinding* ReplaceableBinding = Cast<UMovieSceneReplaceableDirectorBlueprintBinding>(
+							Reference.CustomBinding
+						))
+					{
+						if (ReplaceableBinding->DynamicBinding.Function != nullptr)
+						{
+							return true;
+						}
+					}
+					return false;
 				}
-			}
-			return false;
-			}))
+			))
 		{
 			continue;
 		}
diff --git a/Engine/Plugins/Importers/USDImporter/Source/USDStageImporter/Private/USDStageImportOptions.cpp b/Engine/Plugins/Importers/USDImporter/Source/USDStageImporter/Private/USDStageImportOptions.cpp
--- a/Engine/Plugins/Importers/USDImporter/Source/USDStageImporter/Private/USDStageImportOptions.cpp
+++ b/Engine/Plugins/Importers/USDImporter/Source/USDStageImporter/Private/USDStageImportOptions.cpp
@@ -19,6 +19,7 @@ UUsdStageImportOptions::UUsdStageImportOptions(const FObjectInitializer& ObjectI
 	bImportMaterials = true;
 	bImportGroomAssets = true;
 	bImportSparseVolumeTextures = true;
+	bImportSounds = true;
 	bImportOnlyUsedMaterials = false;
 
 	bUseExistingAssetCache = false;
@@ -97,6 +98,7 @@ void UsdUtils::AddAnalyticsAttributes(const UUsdStageImportOptions& Options, TAr
 	InOutAttributes.Emplace(TEXT("ImportMaterials"), LexToString(Options.bImportMaterials));
 	InOutAttributes.Emplace(TEXT("ImportGroomAssets"), LexToString(Options.bImportGroomAssets));
 	InOutAttributes.Emplace(TEXT("ImportSparseVolumeTextures"), LexToString(Options.bImportSparseVolumeTextures));
+	InOutAttributes.Emplace(TEXT("ImportSounds"), LexToString(Options.bImportSounds));
 	InOutAttributes.Emplace(TEXT("ImportOnlyUsedMaterials"), LexToString(Options.bImportOnlyUsedMaterials));
 	if (Options.PrimsToImport != TArray<FString>{TEXT("/")})
 	{
diff --git a/Engine/Plugins/Importers/USDImporter/Source/USDStageImporter/Private/USDStageImporter.cpp b/Engine/Plugins/Importers/USDImporter/Source/USDStageImporter/Private/USDStageImporter.cpp
--- a/Engine/Plugins/Importers/USDImporter/Source/USDStageImporter/Private/USDStageImporter.cpp
+++ b/Engine/Plugins/Importers/USDImporter/Source/USDStageImporter/Private/USDStageImporter.cpp
@@ -29,6 +29,7 @@
 #include "Animation/Skeleton.h"
 #include "AssetRegistry/AssetRegistryModule.h"
 #include "AssetToolsModule.h"
+#include "AudioDeviceManager.h"
 #include "ComponentRecreateRenderStateContext.h"
 #include "Components/SkeletalMeshComponent.h"
 #include "Components/StaticMeshComponent.h"
@@ -57,6 +58,8 @@
 #include "PackageTools.h"
 #include "PhysicsEngine/PhysicsAsset.h"
 #include "Serialization/ArchiveReplaceObjectRef.h"
+#include "Sound/SoundBase.h"
+#include "Sound/SoundWave.h"
 #include "SparseVolumeTexture/SparseVolumeTexture.h"
 #include "Subsystems/AssetEditorSubsystem.h"
 #include "UObject/UObjectIterator.h"
@@ -428,7 +431,7 @@ namespace UE::USDStageImporter::Private
 			}
 
 #if USE_USD_SDK
-			if (bAnimatedVisibility || UsdUtils::IsAnimated(Prim))
+			if (bAnimatedVisibility || UsdUtils::IsAnimated(Prim) || Prim.IsA(TEXT("SpatialAudio")))
 			{
 				const TOptional<bool> HasAnimatedBounds = {};
 				ImportAnimation(ImportContext, Prim, bAnimatedVisibility, HasAnimatedBounds, Component);
@@ -530,8 +533,8 @@ namespace UE::USDStageImporter::Private
 
 		if (ImportData)
 		{
-			// Don't force update textures as they will already have this preset to their actual texture path
-			if (!Asset->IsA<UTexture>() && !Asset->IsA<USparseVolumeTexture>())
+			// Don't force update these assets as they will already have this preset to their actual source file path
+			if (!Asset->IsA<UTexture>() && !Asset->IsA<USparseVolumeTexture>() && !Asset->IsA<USoundBase>())
 			{
 				ImportData->UpdateFilenameOnly(MainFilePath);
 			}
@@ -612,6 +615,16 @@ namespace UE::USDStageImporter::Private
 			bAssetWasOpen = true;
 		}
 
+		// If the existing asset is an audio file, make sure it's not playing before we stomp it
+		if (USoundWave* Sound = Cast<USoundWave>(ExistingAsset))
+		{
+			if (FAudioDeviceManager* AudioDeviceManager = GEngine->GetAudioDeviceManager())
+			{
+				AudioDeviceManager->StopSoundsUsingResource(Sound);
+			}
+			Sound->FreeResources();
+		}
+
 		UPackage* Package = ExistingPackage ? ExistingPackage : CreatePackage(*TargetPackagePath);
 		if (!Package)
 		{
@@ -801,6 +814,7 @@ namespace UE::USDStageImporter::Private
 		TArray<UObject*> GroomCaches;
 		TArray<UObject*> GroomBindings;
 		TArray<UObject*> SparseVolumeTextures;
+		TArray<UObject*> Sounds;
 
 		TSet<FString> UniqueAssetNames;
 		TMap<UObject*, FString> AssetToContentFolder;
@@ -819,6 +833,7 @@ namespace UE::USDStageImporter::Private
 			const static FString LevelSequencesFolder = TEXT("LevelSequences");
 			const static FString GeometryCachesFolder = TEXT("GeometryCaches");
 			const static FString GroomsFolder = TEXT("Grooms");
+			const static FString SoundsFolder = TEXT("Sounds");
 
 			const FString* AssetTypeFolderPtr = nullptr;
 			if (Asset->IsA(UMaterialInterface::StaticClass()))
@@ -891,6 +906,11 @@ namespace UE::USDStageImporter::Private
 				AssetTypeFolderPtr = &TexturesFolder;
 				SparseVolumeTextures.Add(Asset);
 			}
+			else if (Asset->IsA(USoundBase::StaticClass()))
+			{
+				AssetTypeFolderPtr = &SoundsFolder;
+				Sounds.Add(Asset);
+			}
 			else
 			{
 				// We don't know what to do with this asset
@@ -967,6 +987,7 @@ namespace UE::USDStageImporter::Private
 		PublishAssetType(Materials);
 		PublishAssetType(Textures);
 		PublishAssetType(SparseVolumeTextures);
+		PublishAssetType(Sounds);
 	}
 
 	void ResolveComponentConflict(
@@ -1576,6 +1597,7 @@ namespace UE::USDStageImporter::Private
 			int32 NumGroomBindings = 0;
 			int32 NumGroomCaches = 0;
 			int32 NumSparseVolumeTextures = 0;
+			int32 NumSounds = 0;
 			for (UObject* ImportedAsset : ImportedAssets)
 			{
 				if (!ImportedAsset)
@@ -1627,6 +1649,10 @@ namespace UE::USDStageImporter::Private
 				{
 					++NumSparseVolumeTextures;
 				}
+				else if (ImportedAsset->IsA<USoundBase>())
+				{
+					++NumSounds;
+				}
 			}
 			EventAttributes.Emplace(TEXT("NumStaticMeshes"), LexToString(NumStaticMeshes));
 			EventAttributes.Emplace(TEXT("NumSkeletalMeshes"), LexToString(NumSkeletalMeshes));
@@ -1639,6 +1665,7 @@ namespace UE::USDStageImporter::Private
 			EventAttributes.Emplace(TEXT("NumGroomBindings"), LexToString(NumGroomBindings));
 			EventAttributes.Emplace(TEXT("NumGroomCaches"), LexToString(NumGroomCaches));
 			EventAttributes.Emplace(TEXT("NumSparseVolumeTextures"), LexToString(NumSparseVolumeTextures));
+			EventAttributes.Emplace(TEXT("NumSounds"), LexToString(NumSounds));
 
 			FString RootLayerIdentifier = ImportContext.FilePath;
 			if (ImportContext.Stage)
@@ -1691,7 +1718,8 @@ namespace UE::USDStageImporter::Private
 				(!ImportContext.ImportOptions->bImportLevelSequences && (Asset->IsA<ULevelSequence>())) ||
 				(!ImportContext.ImportOptions->bImportMaterials && (Asset->IsA<UMaterialInterface>() || Asset->IsA<UTexture>())) ||
 				(!ImportContext.ImportOptions->bImportGroomAssets && (Asset->IsA<UGroomAsset>() || Asset->IsA<UGroomCache>() || Asset->IsA<UGroomBindingAsset>())) ||
-				(!ImportContext.ImportOptions->bImportSparseVolumeTextures && (Asset->IsA<USparseVolumeTexture>()))
+				(!ImportContext.ImportOptions->bImportSparseVolumeTextures && (Asset->IsA<USparseVolumeTexture>())) ||
+				(!ImportContext.ImportOptions->bImportSounds && (Asset->IsA<USoundBase>()))
 			)
 			{
 				ObjectsToRemap.Add(Asset, nullptr);
@@ -1924,6 +1952,7 @@ void UUsdStageImporter::ImportFromFile(FUsdStageImportContext& ImportContext)
 														  && ImportContext.ImportOptions->bImportSkeletalAnimations;
 	TranslationContext->bAllowParsingGroomAssets = ImportContext.ImportOptions->bImportGroomAssets;
 	TranslationContext->bAllowParsingSparseVolumeTextures = ImportContext.ImportOptions->bImportSparseVolumeTextures;
+	TranslationContext->bAllowParsingSounds = ImportContext.ImportOptions->bImportSounds;
 	TranslationContext->bTranslateOnlyUsedMaterials = ImportContext.ImportOptions->bImportOnlyUsedMaterials;
 	TranslationContext->InfoCache = InfoCache;
 	TranslationContext->BBoxCache = ImportContext.BBoxCache;
@@ -2129,6 +2158,7 @@ bool UUsdStageImporter::ReimportSingleAsset(
 														  && ImportContext.ImportOptions->bImportSkeletalAnimations;
 	TranslationContext->bAllowParsingGroomAssets = ImportContext.ImportOptions->bImportGroomAssets;
 	TranslationContext->bAllowParsingSparseVolumeTextures = ImportContext.ImportOptions->bImportSparseVolumeTextures;
+	TranslationContext->bAllowParsingSounds = ImportContext.ImportOptions->bImportSounds;
 	TranslationContext->bTranslateOnlyUsedMaterials = ImportContext.ImportOptions->bImportOnlyUsedMaterials;
 	TranslationContext->InfoCache = InfoCache;
 	TranslationContext->BBoxCache = ImportContext.BBoxCache;
diff --git a/Engine/Plugins/Importers/USDImporter/Source/USDStageImporter/Public/USDStageImportOptions.h b/Engine/Plugins/Importers/USDImporter/Source/USDStageImporter/Public/USDStageImportOptions.h
--- a/Engine/Plugins/Importers/USDImporter/Source/USDStageImporter/Public/USDStageImportOptions.h
+++ b/Engine/Plugins/Importers/USDImporter/Source/USDStageImporter/Public/USDStageImportOptions.h
@@ -78,6 +78,10 @@ public:
 	UPROPERTY(BlueprintReadWrite, config, EditAnywhere, Category = "DataToImport", meta = (DisplayName = "Sparse Volume Textures"))
 	bool bImportSparseVolumeTextures;
 
+	/** Whether to import audio files referenced by UsdMediaSpatialAudio schemas as Unreal sound assets */
+	UPROPERTY(BlueprintReadWrite, config, EditAnywhere, Category = "DataToImport", meta = (DisplayName = "Audio"))
+	bool bImportSounds;
+
 	/**
 	 * If this is checked, only materials actively used by the stage and import settings will be parsed.
 	 * If this is unchecked, all materials present on the stage will be parsed.
diff --git a/Engine/Plugins/Runtime/USDCore/Source/USDClasses/Private/USDClassesModule.cpp b/Engine/Plugins/Runtime/USDCore/Source/USDClasses/Private/USDClassesModule.cpp
--- a/Engine/Plugins/Runtime/USDCore/Source/USDClasses/Private/USDClassesModule.cpp
+++ b/Engine/Plugins/Runtime/USDCore/Source/USDClasses/Private/USDClassesModule.cpp
@@ -39,6 +39,7 @@
 #include "Modules/ModuleManager.h"
 #include "PhysicsEngine/PhysicsAsset.h"
 #include "Serialization/JsonSerializer.h"
+#include "Sound/SoundWave.h"
 #include "SparseVolumeTexture/SparseVolumeTexture.h"
 #include "UObject/NameTypes.h"
 #include "UObject/ObjectSaveContext.h"
@@ -448,6 +449,10 @@ TSet<UObject*> IUsdClassesModule::GetAssetDependencies(UObject* Asset)
 	{
 		// Do nothing. SparseVolumeTextures have no additional dependencies
 	}
+	else if (USoundBase* Sound = Cast<USoundBase>(Asset))
+	{
+		// Do nothing. Sounds have no additional dependencies
+	}
 	else
 	{
 		UE_LOG(LogUsd, Warning, TEXT("Unknown asset '%s' encountered when collecting dependencies."), Asset ? *Asset->GetName() : TEXT("nullptr"));
diff --git a/Engine/Plugins/Runtime/USDCore/Source/USDClasses/Public/USDProjectSettings.h b/Engine/Plugins/Runtime/USDCore/Source/USDClasses/Public/USDProjectSettings.h
--- a/Engine/Plugins/Runtime/USDCore/Source/USDClasses/Public/USDProjectSettings.h
+++ b/Engine/Plugins/Runtime/USDCore/Source/USDClasses/Public/USDProjectSettings.h
@@ -109,7 +109,10 @@ public:
 	UPROPERTY(config, EditAnywhere, Category = "USD|Dialogs")
 	EUsdSaveDialogBehavior ShowSaveLayersDialogWhenClosing = EUsdSaveDialogBehavior::ShowPrompt;
 
-	// Note that the below properties being FSoftObjectPath ensure that these materials are cooked into packaged games
+	// Note that the below properties being FSoftObjectPath ensure that these assets are cooked into packaged games
+
+	UPROPERTY(config, EditAnywhere, Category = USD, meta = (AllowedClasses = "/Script/Engine.SoundAttenuation"))
+	FSoftObjectPath DefaultSoundAttenuation = FSoftObjectPath{TEXT("/USDCore/USDDefaultAttenuation.USDDefaultAttenuation")};
 
 	/**
 	 * Material to use when handling .vdb files as Sparse Volume Textures. An instance of this material will be
diff --git a/Engine/Plugins/Runtime/USDCore/Source/USDUtilities/Private/USDConversionUtils.cpp b/Engine/Plugins/Runtime/USDCore/Source/USDUtilities/Private/USDConversionUtils.cpp
--- a/Engine/Plugins/Runtime/USDCore/Source/USDUtilities/Private/USDConversionUtils.cpp
+++ b/Engine/Plugins/Runtime/USDCore/Source/USDUtilities/Private/USDConversionUtils.cpp
@@ -26,6 +26,7 @@
 #include "Animation/Skeleton.h"
 #include "CineCameraActor.h"
 #include "CineCameraComponent.h"
+#include "Components/AudioComponent.h"
 #include "Components/DirectionalLightComponent.h"
 #include "Components/HeterogeneousVolumeComponent.h"
 #include "Components/PointLightComponent.h"
@@ -50,6 +51,7 @@
 #include "LandscapeProxy.h"
 #include "Misc/PackageName.h"
 #include "PhysicsEngine/PhysicsAsset.h"
+#include "Sound/AmbientSound.h"
 #include "SparseVolumeTexture/SparseVolumeTexture.h"
 #include "Widgets/Notifications/SNotificationList.h"
 
@@ -81,6 +83,7 @@
 #include "pxr/usd/usdLux/rectLight.h"
 #include "pxr/usd/usdLux/shapingAPI.h"
 #include "pxr/usd/usdLux/sphereLight.h"
+#include "pxr/usd/usdMedia/spatialAudio.h"
 #include "pxr/usd/usdSkel/animation.h"
 #include "pxr/usd/usdSkel/cache.h"
 #include "pxr/usd/usdSkel/root.h"
@@ -473,6 +476,10 @@ UClass* UsdUtils::GetActorTypeForPrim(const pxr::UsdPrim& Prim)
 	{
 		return AHeterogeneousVolume::StaticClass();
 	}
+	else if (Prim.IsA<pxr::UsdMediaSpatialAudio>())
+	{
+		return AAmbientSound::StaticClass();
+	}
 	else
 	{
 		return AActor::StaticClass();
@@ -580,6 +587,10 @@ FString UsdUtils::GetSchemaNameForComponent(const USceneComponent& Component)
 	{
 		return TEXT("DomeLight");
 	}
+	else if (Component.IsA<UAudioComponent>())
+	{
+		return TEXT("SpatialAudio");
+	}
 
 	return TEXT("Xform");
 }
diff --git a/Engine/Plugins/Runtime/USDCore/Source/USDUtilities/Private/USDPrimConversion.cpp b/Engine/Plugins/Runtime/USDCore/Source/USDUtilities/Private/USDPrimConversion.cpp
--- a/Engine/Plugins/Runtime/USDCore/Source/USDUtilities/Private/USDPrimConversion.cpp
+++ b/Engine/Plugins/Runtime/USDCore/Source/USDUtilities/Private/USDPrimConversion.cpp
@@ -28,6 +28,7 @@
 #include "Channels/MovieSceneFloatChannel.h"
 #include "CineCameraActor.h"
 #include "CineCameraComponent.h"
+#include "Components/AudioComponent.h"
 #include "Components/BrushComponent.h"
 #include "Components/DirectionalLightComponent.h"
 #include "Components/HeterogeneousVolumeComponent.h"
@@ -53,11 +54,14 @@
 #include "MovieSceneTimeHelpers.h"
 #include "Rendering/SkeletalMeshRenderData.h"
 #include "Sections/MovieScene3DTransformSection.h"
+#include "Sections/MovieSceneAudioSection.h"
 #include "Sections/MovieSceneBoolSection.h"
 #include "Sections/MovieSceneColorSection.h"
 #include "Sections/MovieSceneFloatSection.h"
 #include "Sections/MovieSceneVectorSection.h"
 #include "Sections/MovieSceneVisibilitySection.h"
+#include "Sound/SoundBase.h"
+#include "Sound/SoundWave.h"
 #include "SparseVolumeTexture/SparseVolumeTexture.h"
 #include "Tracks/MovieScene3DTransformTrack.h"
 #include "Tracks/MovieSceneBoolTrack.h"
@@ -104,6 +108,7 @@
 #include "pxr/usd/usdLux/shapingAPI.h"
 #include "pxr/usd/usdLux/sphereLight.h"
 #include "pxr/usd/usdLux/tokens.h"
+#include "pxr/usd/usdMedia/spatialAudio.h"
 #include "pxr/usd/usdShade/connectableAPI.h"
 #include "pxr/usd/usdShade/material.h"
 #include "pxr/usd/usdShade/materialBindingAPI.h"
@@ -289,6 +294,28 @@ namespace UE::USDPrimConversion::Private
 			OutTransform = FTransform{OutTransform.ToMatrixWithScale() * ParentToWorld.ToInverseMatrixWithScale()};
 		}
 	}
+
+	void PrepareAttributeForAnimationBake(const pxr::UsdAttribute& Attr)
+	{
+		if (!Attr)
+		{
+			return;
+		}
+
+		// Weirdly enough GetTimeSamples() will return time codes with the offset and scale applied, while
+		// ClearAtTime() expects time codes without offset and scale applied, so we must manually undo them here
+		UE::FSdfLayerOffset CombinedOffset = UsdUtils::GetPrimToStageOffset(UE::FUsdPrim{Attr.GetPrim()});
+
+		std::vector<double> TimeSamples;
+		Attr.GetTimeSamples(&TimeSamples);
+		for (double TimeSample : TimeSamples)
+		{
+			double LocalTime = (TimeSample - CombinedOffset.Offset) / CombinedOffset.Scale;
+			Attr.ClearAtTime(LocalTime);
+		}
+
+		UsdUtils::NotifyIfOverriddenOpinion(Attr);
+	}
 }	 // namespace UE::USDPrimConversion::Private
 
 bool UsdToUnreal::ConvertXformable(
@@ -655,7 +682,8 @@ bool UsdToUnreal::ConvertFloatTimeSamples(
 	const UE::FUsdStage& Stage,
 	const TArray<double>& UsdTimeSamples,
 	const TFunction<float(double)>& ReaderFunc,
-	UMovieSceneFloatTrack& MovieSceneTrack,
+	FMovieSceneFloatChannel& FloatChannel,
+	const UMovieScene& MovieSceneOuter,
 	const FMovieSceneSequenceTransform& SequenceTransform,
 	TOptional<ERichCurveInterpMode> InterpolationModeOverride
 )
@@ -665,14 +693,8 @@ bool UsdToUnreal::ConvertFloatTimeSamples(
 		return false;
 	}
 
-	const UMovieScene* MovieScene = MovieSceneTrack.GetTypedOuter<UMovieScene>();
-	if (!MovieScene)
-	{
-		return false;
-	}
-
-	const FFrameRate Resolution = MovieScene->GetTickResolution();
-	const FFrameRate DisplayRate = MovieScene->GetDisplayRate();
+	const FFrameRate Resolution = MovieSceneOuter.GetTickResolution();
+	const FFrameRate DisplayRate = MovieSceneOuter.GetDisplayRate();
 
 	FScopedUsdAllocs Allocs;
 
@@ -716,6 +738,31 @@ bool UsdToUnreal::ConvertFloatTimeSamples(
 		SectionValues.Emplace_GetRef(UEValue).InterpMode = InterpMode;
 	}
 
+	FloatChannel.Set(FrameNumbers, SectionValues);
+
+	return true;
+}
+
+bool UsdToUnreal::ConvertFloatTimeSamples(
+	const UE::FUsdStage& Stage,
+	const TArray<double>& UsdTimeSamples,
+	const TFunction<float(double)>& ReaderFunc,
+	UMovieSceneFloatTrack& MovieSceneTrack,
+	const FMovieSceneSequenceTransform& SequenceTransform,
+	TOptional<ERichCurveInterpMode> InterpolationModeOverride
+)
+{
+	if (!ReaderFunc)
+	{
+		return false;
+	}
+
+	const UMovieScene* MovieScene = MovieSceneTrack.GetTypedOuter<UMovieScene>();
+	if (!MovieScene)
+	{
+		return false;
+	}
+
 	bool bSectionAdded = false;
 	UMovieSceneFloatSection* Section = Cast<UMovieSceneFloatSection>(MovieSceneTrack.FindOrAddSection(0, bSectionAdded));
 	Section->EvalOptions.CompletionMode = EMovieSceneCompletionMode::KeepState;
@@ -723,11 +770,25 @@ bool UsdToUnreal::ConvertFloatTimeSamples(
 	TArrayView<FMovieSceneFloatChannel*> Channels = Section->GetChannelProxy().GetChannels<FMovieSceneFloatChannel>();
 	if (Channels.Num() > 0)
 	{
-		Channels[0]->Set(FrameNumbers, SectionValues);
+		FMovieSceneFloatChannel* Channel = Channels[0];
+
+		const bool bSuccess = UsdToUnreal::ConvertFloatTimeSamples(
+			Stage,
+			UsdTimeSamples,
+			ReaderFunc,
+			*Channel,
+			*MovieScene,
+			SequenceTransform,
+			InterpolationModeOverride
+		);
+
+		if (!bSuccess)
+		{
+			return false;
+		}
 	}
 
 	Section->SetRange(Section->GetAutoSizeRange().Get(TRange<FFrameNumber>::Empty()));
-
 	return true;
 }
 
@@ -1664,6 +1725,27 @@ UsdToUnreal::FPropertyTrackReader UsdToUnreal::CreatePropertyTrackReader(
 			}
 		}
 	}
+	else if (const pxr::UsdMediaSpatialAudio Audio{Prim})
+	{
+		if (PropertyPath == TEXT("Volume"))
+		{
+			if (pxr::UsdAttribute Attr = Audio.GetGainAttr())
+			{
+				double Default = 1.0;
+				Attr.Get<double>(&Default);
+
+				// The VolumeMultiplier property is a float anyway, so we may as well convert
+				// doubles to floats right here
+				Reader.FloatReader = [Attr, Default](double UsdTimeCode)
+				{
+					double Result = Default;
+					Attr.Get<double>(&Result, UsdTimeCode);
+					return static_cast<float>(FMath::Max(Result, 0.0));
+				};
+				return Reader;
+			}
+		}
+	}
 
 	return Reader;
 }
@@ -2050,6 +2132,325 @@ bool UnrealToUsd::ConvertCameraComponent(const UCineCameraComponent& CameraCompo
 	return true;
 }
 
+bool UnrealToUsd::ConvertAudioComponent(const UAudioComponent& AudioComponent, pxr::UsdPrim& Prim, bool bFilePathOnly, double UsdTimeCode)
+{
+	FScopedUsdAllocs UsdAllocs;
+
+	pxr::UsdMediaSpatialAudio SpatialAudio{Prim};
+	if (!SpatialAudio)
+	{
+		return false;
+	}
+
+	if (UsdUtils::NotifyIfInstanceProxy(Prim))
+	{
+		return false;
+	}
+
+	if (!bFilePathOnly)
+	{
+		// auralmode
+		if (pxr::UsdAttribute Attr = SpatialAudio.CreateAuralModeAttr())
+		{
+			const bool bIsSpatial = AudioComponent.bAllowSpatialization
+									&& (AudioComponent.bOverrideAttenuation || AudioComponent.AttenuationSettings);
+
+			Attr.Set<pxr::TfToken>(bIsSpatial ? pxr::UsdMediaTokens->spatial : pxr::UsdMediaTokens->nonSpatial);
+			UsdUtils::NotifyIfOverriddenOpinion(Attr);
+		}
+
+		// gain
+		if (pxr::UsdAttribute Attr = SpatialAudio.CreateGainAttr())
+		{
+			Attr.Set<double>(static_cast<double>(AudioComponent.VolumeMultiplier), UsdTimeCode);
+			UsdUtils::NotifyIfOverriddenOpinion(Attr);
+		}
+
+		// playbackMode
+		if (pxr::UsdAttribute Attr = SpatialAudio.CreatePlaybackModeAttr())
+		{
+			if (USoundWave* Sound = Cast<USoundWave>(AudioComponent.Sound))	   // GetSound() is not const so we can't use it
+			{
+				const bool bIsLooping = Sound->IsLooping();
+
+				Attr.Set<pxr::TfToken>(bIsLooping ? pxr::UsdMediaTokens->loopFromStart : pxr::UsdMediaTokens->onceFromStart);
+				UsdUtils::NotifyIfOverriddenOpinion(Attr);
+			}
+		}
+
+		// startTime
+		// We don't really want to author anything here, but since there is no concept of "startTime" on the UE side
+		// we really want this audio to play at startTime zero, so if for some reason this prim already has an opinion
+		// otherwise we need to override it (note how we're just using GetAttr instead of CreateAttr).
+		if (pxr::UsdAttribute Attr = SpatialAudio.GetStartTimeAttr())
+		{
+			Attr.Set<pxr::SdfTimeCode>(0.0);
+			UsdUtils::NotifyIfOverriddenOpinion(Attr);
+		}
+	}
+
+#if WITH_EDITOR
+	// filePath
+	if (pxr::UsdAttribute Attr = SpatialAudio.CreateFilePathAttr())
+	{
+		UsdUtils::NotifyIfOverriddenOpinion(Attr);
+
+		FString FilePath;
+		if (USoundWave* Sound = Cast<USoundWave>(AudioComponent.Sound))	   // GetSound() is not const so we can't use it
+		{
+			if (UAssetImportData* ImportData = Sound->AssetImportData)
+			{
+				FilePath = ImportData->GetFirstFilename();
+			}
+		}
+
+		if (!FilePath.IsEmpty())
+		{
+			if (!FPaths::FileExists(FilePath))
+			{
+				UE_LOG(
+					LogUsd,
+					Warning,
+					TEXT(
+						"Used '%s' as the audio file when converting AudioComponent '%s' onto prim '%s', but the file does not exist on the filesystem!"
+					),
+					*FilePath,
+					*AudioComponent.GetPathName(),
+					*UsdToUnreal::ConvertPath(Prim.GetPrimPath())
+				);
+			}
+
+			pxr::SdfAssetPath AssetPath{UnrealToUsd::ConvertString(*FilePath).Get()};
+			Attr.Set(AssetPath);
+		}
+		else
+		{
+			Attr.ClearAtTime(pxr::UsdTimeCode::Default());
+		}
+	}
+#endif	  // WITH_EDITOR
+
+	return true;
+}
+
+bool UnrealToUsd::ConvertAudioSection(
+	const UMovieSceneAudioSection& AudioSection,
+	const FMovieSceneSequenceTransform& SequenceTransform,
+	pxr::UsdPrim& Prim
+)
+{
+	FScopedUsdAllocs UsdAllocs;
+
+	pxr::UsdMediaSpatialAudio SpatialAudio{Prim};
+	if (!SpatialAudio)
+	{
+		return false;
+	}
+
+	if (UsdUtils::NotifyIfInstanceProxy(Prim))
+	{
+		return false;
+	}
+
+	UMovieScene* MovieScene = AudioSection.GetTypedOuter<UMovieScene>();
+	if (!MovieScene)
+	{
+		return false;
+	}
+
+	const double StageTimeCodesPerSecond = Prim.GetStage()->GetTimeCodesPerSecond();
+	const FFrameRate StageFrameRate(StageTimeCodesPerSecond, 1);
+
+#if WITH_EDITOR
+	// filePath
+	if (pxr::UsdAttribute Attr = SpatialAudio.CreateFilePathAttr())
+	{
+		FString FilePath;
+		if (const USoundWave* Sound = Cast<USoundWave>(AudioSection.GetSound()))
+		{
+			if (UAssetImportData* ImportData = Sound->AssetImportData)
+			{
+				FilePath = ImportData->GetFirstFilename();
+			}
+		}
+
+		if (!FilePath.IsEmpty())
+		{
+			// Don't author anything if we're just trying to set a relative path version
+			// of the same file that is currently set as an absolute path (or vice-versa)
+			bool bSetNewPath = true;
+			if (Attr.HasAuthoredValue())
+			{
+				FString CurrentPath = UsdUtils::GetResolvedAssetPath(Attr, pxr::UsdTimeCode::Default());
+
+				// Both paths should be absolute at this point
+				bSetNewPath = !FPaths::IsSamePath(FilePath, CurrentPath);
+			}
+
+			if (bSetNewPath)
+			{
+				UsdUtils::NotifyIfOverriddenOpinion(Attr);
+				pxr::SdfAssetPath AssetPath{UnrealToUsd::ConvertString(*FilePath).Get()};
+				Attr.Set(AssetPath);
+			}
+		}
+		else
+		{
+			UsdUtils::NotifyIfOverriddenOpinion(Attr);
+			Attr.ClearAtTime(pxr::UsdTimeCode::Default());
+		}
+	}
+#endif	  // WITH_EDITOR
+
+	// mediaOffset
+	if (pxr::UsdAttribute Attr = SpatialAudio.CreateMediaOffsetAttr())
+	{
+		const FFrameNumber& StartOffset = AudioSection.GetStartOffset();
+
+		UsdUtils::NotifyIfOverriddenOpinion(Attr);
+
+		const FFrameRate Resolution = MovieScene->GetTickResolution();
+		double OffsetSeconds = Resolution.AsSeconds(StartOffset);
+
+		Attr.Set(OffsetSeconds);
+	}
+
+	// startTime and endTime
+	double bChangedTimes = false;
+	pxr::UsdAttribute StartAttr = SpatialAudio.CreateStartTimeAttr();
+	pxr::UsdAttribute EndAttr = SpatialAudio.CreateEndTimeAttr();
+	if (StartAttr && EndAttr)
+	{
+		UsdUtils::NotifyIfOverriddenOpinion(StartAttr);
+		UsdUtils::NotifyIfOverriddenOpinion(EndAttr);
+
+		const FFrameRate Resolution = MovieScene->GetTickResolution();
+		const FFrameRate DisplayRate = MovieScene->GetDisplayRate();
+
+		TRange<FFrameNumber> Range = AudioSection.GetRange();
+		const FFrameNumber& StartTick = Range.GetLowerBoundValue();
+		const FFrameNumber& EndTick = Range.GetUpperBoundValue();
+
+		FFrameTime UsdStartTickTime = FFrameRate::Snap(StartTick, Resolution, DisplayRate).FloorToFrame();
+		FFrameTime UsdEndTickTime = FFrameRate::Snap(EndTick, Resolution, DisplayRate).FloorToFrame();
+		UsdStartTickTime *= SequenceTransform.InverseNoLooping();
+		UsdEndTickTime *= SequenceTransform.InverseNoLooping();
+		double UsdStartTimeCode = FFrameRate::TransformTime(UsdStartTickTime, Resolution, StageFrameRate).AsDecimal();
+		double UsdEndTimeCode = FFrameRate::TransformTime(UsdEndTickTime, Resolution, StageFrameRate).AsDecimal();
+
+		// Get the previous start/end times from the attribute
+		bool bGotPreviousStart = false;
+		bool bGotPreviousEnd = false;
+		pxr::SdfTimeCode PreviousStartTimeCode = 0.0;
+		pxr::SdfTimeCode PreviousEndTimeCode = 0.0;
+		if (StartAttr.Get(&PreviousStartTimeCode))
+		{
+			bGotPreviousStart = true;
+		}
+		if (EndAttr.Get(&PreviousEndTimeCode))
+		{
+			bGotPreviousEnd = true;
+		}
+
+		StartAttr.Set(pxr::SdfTimeCode{UsdStartTimeCode});
+		EndAttr.Set(pxr::SdfTimeCode{UsdEndTimeCode});
+
+		// Retrieve the new values from the attribute and record whether they changed or not.
+		//
+		// We're not just using the values we just set above (e.g. UsdStartTimeCode) because
+		// USD does automatic layer offset/scale conversions for these, that we would need to manually
+		// undo here anyway in order to properly compare with the previous values.
+		// We may as well just let USD do them by calling Get() instead, and then we can know
+		// for sure whether our actual set values changed
+		pxr::SdfTimeCode SetStartTimeCode = 0.0;
+		pxr::SdfTimeCode SetEndTimeCode = 0.0;
+		if (bGotPreviousStart && StartAttr.Get(&SetStartTimeCode))
+		{
+			if (!FMath::IsNearlyEqual(PreviousStartTimeCode.GetValue(), SetStartTimeCode.GetValue()))
+			{
+				bChangedTimes = true;
+			}
+		}
+		if (!bChangedTimes && bGotPreviousEnd && EndAttr.Get(&SetEndTimeCode))
+		{
+			if (!FMath::IsNearlyEqual(PreviousEndTimeCode.GetValue(), SetEndTimeCode.GetValue()))
+			{
+				bChangedTimes = true;
+			}
+		}
+	}
+
+	// playbackMode
+	// We lose some "degrees of freedom" here by just authoring 2 out of the 5 possible values, but in Unreal
+	// the audio section must always have closed start and end frames, which means only the "FromStartToEnd"
+	// options really make sense. It would probably be more confusing to try and come up with some heuristics
+	// as to when we should guess that the playbackMode really should be something else.
+	// As a small concession though, we'll only change the playback mode in case we have some new value for
+	// start/endTime/looping, so that the user can keep their custom playbackMode until we have to actually
+	// change it
+	if (pxr::UsdAttribute Attr = SpatialAudio.CreatePlaybackModeAttr())
+	{
+		bool bPreviousLooping = false;
+		pxr::TfToken CurrentPlaybackMode;
+		if (Attr.Get(&CurrentPlaybackMode))
+		{
+			bPreviousLooping = (CurrentPlaybackMode != pxr::UsdMediaTokens->onceFromStart)
+							   && (CurrentPlaybackMode != pxr::UsdMediaTokens->onceFromStartToEnd);
+		}
+
+		const bool bNowLooping = AudioSection.GetLooping();
+		const bool bChangedLooping = bNowLooping ^ bPreviousLooping;
+
+		if (bChangedTimes || bChangedLooping)
+		{
+			UsdUtils::NotifyIfOverriddenOpinion(Attr);
+
+			if (bNowLooping)
+			{
+				Attr.Set(pxr::UsdMediaTokens->loopFromStartToEnd);
+			}
+			else
+			{
+				Attr.Set(pxr::UsdMediaTokens->onceFromStartToEnd);
+			}
+		}
+	}
+
+	// auralmode
+	if (pxr::UsdAttribute Attr = SpatialAudio.CreateAuralModeAttr())
+	{
+		const bool bIsSpatial = AudioSection.GetOverrideAttenuation() && AudioSection.GetAttenuationSettings();
+		Attr.Set<pxr::TfToken>(bIsSpatial ? pxr::UsdMediaTokens->spatial : pxr::UsdMediaTokens->nonSpatial);
+		UsdUtils::NotifyIfOverriddenOpinion(Attr);
+	}
+
+	// gain
+	if (pxr::UsdAttribute Attr = SpatialAudio.CreateGainAttr())
+	{
+		const FMovieSceneFloatChannel& VolumeChannel = AudioSection.GetSoundVolumeChannel();
+
+		// Write out default value
+		// Note that even the default is only set on the section and not the component. This because setting the volume
+		// on the component itself will do precisely nothing if the Sequencer is usually the source of the audio anyway
+		if (TOptional<float> Default = VolumeChannel.GetDefault())
+		{
+			Attr.Set(static_cast<double>(Default.GetValue()));
+		}
+
+		UE::USDPrimConversion::Private::PrepareAttributeForAnimationBake(Attr);
+
+		// Write out timeSamples
+		TFunction<void(float, double)> BakerFunc = [Attr](float UEValue, double UsdTimeCode) mutable
+		{
+			Attr.Set(static_cast<double>(UEValue), UsdTimeCode);
+		};
+		UE::FUsdPrim WrappedPrim{Prim};
+		UnrealToUsd::ConvertFloatChannel(VolumeChannel, *MovieScene, SequenceTransform, BakerFunc, WrappedPrim);
+	}
+
+	return true;
+}
+
 bool UnrealToUsd::ConvertBoolTrack(
 	const UMovieSceneBoolTrack& MovieSceneTrack,
 	const FMovieSceneSequenceTransform& SequenceTransform,
@@ -2121,8 +2522,9 @@ bool UnrealToUsd::ConvertBoolTrack(
 	return true;
 }
 
-bool UnrealToUsd::ConvertFloatTrack(
-	const UMovieSceneFloatTrack& MovieSceneTrack,
+bool UnrealToUsd::ConvertFloatChannel(
+	const FMovieSceneFloatChannel& MovieSceneChannel,
+	const UMovieScene& MovieSceneOuter,
 	const FMovieSceneSequenceTransform& SequenceTransform,
 	const TFunction<void(float, double)>& WriterFunc,
 	UE::FUsdPrim& Prim
@@ -2133,12 +2535,6 @@ bool UnrealToUsd::ConvertFloatTrack(
 		return false;
 	}
 
-	UMovieScene* MovieScene = MovieSceneTrack.GetTypedOuter<UMovieScene>();
-	if (!MovieScene)
-	{
-		return false;
-	}
-
 	UE::FUsdStage Stage = Prim.GetStage();
 
 	ERichCurveInterpMode StageInterpMode;
@@ -2149,9 +2545,9 @@ bool UnrealToUsd::ConvertFloatTrack(
 																								: ERichCurveInterpMode::RCIM_Constant;
 	}
 
-	const TRange<FFrameNumber> PlaybackRange = MovieScene->GetPlaybackRange();
-	const FFrameRate Resolution = MovieScene->GetTickResolution();
-	const FFrameRate DisplayRate = MovieScene->GetDisplayRate();
+	const TRange<FFrameNumber> PlaybackRange = MovieSceneOuter.GetPlaybackRange();
+	const FFrameRate Resolution = MovieSceneOuter.GetTickResolution();
+	const FFrameRate DisplayRate = MovieSceneOuter.GetDisplayRate();
 
 	const double StageTimeCodesPerSecond = Stage.GetTimeCodesPerSecond();
 	const FFrameRate StageFrameRate(StageTimeCodesPerSecond, 1);
@@ -2206,24 +2602,38 @@ bool UnrealToUsd::ConvertFloatTrack(
 		return Values;
 	};
 
-	for (UMovieSceneSection* Section : MovieSceneTrack.GetAllSections())
+	for (const TPair<FFrameNumber, float>& Pair : EvaluateChannel(MovieSceneChannel, 0.0f))
 	{
-		if (UMovieSceneFloatSection* FloatSection = Cast<UMovieSceneFloatSection>(Section))
-		{
-			for (const TPair<FFrameNumber, float>& Pair : EvaluateChannel(FloatSection->GetChannel(), 0.0f))
-			{
-				FFrameTime TransformedBakedKeyTime{Pair.Key};
-				TransformedBakedKeyTime *= SequenceTransform.InverseNoLooping();
-				FFrameTime UsdFrameTime = FFrameRate::TransformTime(TransformedBakedKeyTime, Resolution, StageFrameRate);
+		FFrameTime TransformedBakedKeyTime{Pair.Key};
+		TransformedBakedKeyTime *= SequenceTransform.InverseNoLooping();
+		FFrameTime UsdFrameTime = FFrameRate::TransformTime(TransformedBakedKeyTime, Resolution, StageFrameRate);
 
-				float UEValue = Pair.Value;
+		float UEValue = Pair.Value;
 
-				WriterFunc(UEValue, UsdFrameTime.AsDecimal());
+		WriterFunc(UEValue, UsdFrameTime.AsDecimal());
+	}
+	return true;
+}
+
+bool UnrealToUsd::ConvertFloatTrack(
+	const UMovieSceneFloatTrack& MovieSceneTrack,
+	const FMovieSceneSequenceTransform& SequenceTransform,
+	const TFunction<void(float, double)>& WriterFunc,
+	UE::FUsdPrim& Prim
+)
+{
+	if (UMovieScene* MovieScene = MovieSceneTrack.GetTypedOuter<UMovieScene>())
+	{
+		for (UMovieSceneSection* Section : MovieSceneTrack.GetAllSections())
+		{
+			if (UMovieSceneFloatSection* FloatSection = Cast<UMovieSceneFloatSection>(Section))
+			{
+				return ConvertFloatChannel(FloatSection->GetChannel(), *MovieScene, SequenceTransform, WriterFunc, Prim);
 			}
 		}
 	}
 
-	return true;
+	return false;
 }
 
 bool UnrealToUsd::ConvertColorTrack(
@@ -4689,28 +5099,14 @@ UnrealToUsd::FPropertyTrackWriter UnrealToUsd::CreatePropertyTrackWriter(
 
 	for (const pxr::UsdAttribute& SomeAttr : Attrs)
 	{
-		if (SomeAttr)
-		{
-			// Weirdly enough GetTimeSamples() will return time codes with the offset and scale applied, while
-			// ClearAtTime() expects time codes without offset and scale applied, so we must manually undo them here
-			UE::FSdfLayerOffset CombinedOffset = UsdUtils::GetPrimToStageOffset(UE::FUsdPrim{SomeAttr.GetPrim()});
-
-			std::vector<double> TimeSamples;
-			SomeAttr.GetTimeSamples(&TimeSamples);
-			for (double TimeSample : TimeSamples)
-			{
-				double LocalTime = (TimeSample - CombinedOffset.Offset) / CombinedOffset.Scale;
-				SomeAttr.ClearAtTime(LocalTime);
-			}
-
-			// Note that we must do this only after the change block is destroyed!
-			// This is important because if we don't have spec for this attribute on the current edit target, we're relying
-			// on the previous code to create it, and we need to let USD emit its internal notices and fully commit the
-			// "attribute creation" spec first. This because NotifyIfOverriddenOpinion will go through the attribute's
-			// spec stack and consider our attribute overriden if it finds a stronger opinion than the one on the edit
-			// target. Well if our own spec hasn't been created yet it will misfire when it runs into any other spec
-			UsdUtils::NotifyIfOverriddenOpinion(SomeAttr);
-		}
+		// Note that we must do this only after the change block is destroyed!
+		// This is important because if we don't have spec for this attribute on the current edit target, we're relying
+		// on the previous code to create it, and we need to let USD emit its internal notices and fully commit the
+		// "attribute creation" spec first. This because PrepareAttributeForAnimationBake will call NotifyIfOverriddenOpinion,
+		// which will go through the attribute's spec stack and consider our attribute overriden if it finds a stronger
+		// opinion than the one on the edit target. Well if our own spec hasn't been created yet it will misfire when
+		// it runs into any other spec
+		UE::USDPrimConversion::Private::PrepareAttributeForAnimationBake(SomeAttr);
 	}
 
 	return Result;
@@ -5142,6 +5538,13 @@ TArray<UE::FUsdAttribute> UnrealToUsd::GetAttributesForProperty(const UE::FUsdPr
 			return Attrs;
 		}
 	}
+	else if (PropertyPath == GET_MEMBER_NAME_CHECKED(UAudioComponent, VolumeMultiplier))
+	{
+		if (pxr::UsdMediaSpatialAudio Audio{Prim})
+		{
+			return {UE::FUsdAttribute{Audio.GetGainAttr()}};
+		}
+	}
 
 	return {};
 }
diff --git a/Engine/Plugins/Runtime/USDCore/Source/USDUtilities/Public/USDPrimConversion.h b/Engine/Plugins/Runtime/USDCore/Source/USDUtilities/Public/USDPrimConversion.h
--- a/Engine/Plugins/Runtime/USDCore/Source/USDUtilities/Public/USDPrimConversion.h
+++ b/Engine/Plugins/Runtime/USDCore/Source/USDUtilities/Public/USDPrimConversion.h
@@ -29,12 +29,14 @@ PXR_NAMESPACE_OPEN_SCOPE
 PXR_NAMESPACE_CLOSE_SCOPE
 
 class AInstancedFoliageActor;
+class UAudioComponent;
 class UCineCameraComponent;
 class UHierarchicalInstancedStaticMeshComponent;
 class ULevel;
 class UMeshComponent;
 class UMovieScene;
 class UMovieScene3DTransformTrack;
+class UMovieSceneAudioSection;
 class UMovieSceneBoolTrack;
 class UMovieSceneColorTrack;
 class UMovieSceneDoubleVectorTrack;
@@ -48,6 +50,7 @@ class USkeletalMeshComponent;
 class UUsdDrawModeComponent;
 enum ERichCurveInterpMode : int;
 struct FFrameRate;
+struct FMovieSceneFloatChannel;
 struct FMovieSceneSequenceTransform;
 struct FUsdCombinedPrimMetadata;
 struct FUsdPrimMetadata;
@@ -136,6 +139,15 @@ namespace UsdToUnreal
 		UMovieSceneVisibilityTrack& MovieSceneTrack,
 		const FMovieSceneSequenceTransform& SequenceTransform
 	);
+	USDUTILITIES_API bool ConvertFloatTimeSamples(
+		const UE::FUsdStage& UsdStage,
+		const TArray<double>& UsdTimeSamples,
+		const TFunction<float(double)>& ReaderFunc,
+		FMovieSceneFloatChannel& FloatChannel,
+		const UMovieScene& MovieSceneOuter,
+		const FMovieSceneSequenceTransform& SequenceTransform,
+		TOptional<ERichCurveInterpMode> InterpolationModeOverride = {}
+	);
 	USDUTILITIES_API bool ConvertFloatTimeSamples(
 		const UE::FUsdStage& UsdStage,
 		const TArray<double>& UsdTimeSamples,
@@ -235,11 +247,31 @@ namespace UnrealToUsd
 		double UsdTimeCode = UsdUtils::GetDefaultTimeCode()
 	);
 
+	USDUTILITIES_API bool ConvertAudioComponent(
+		const UAudioComponent& AudioComponent,
+		pxr::UsdPrim& Prim,
+		bool bFilePathOnly,
+		double UsdTimeCode = UsdUtils::GetDefaultTimeCode()
+	);
+
+	USDUTILITIES_API bool ConvertAudioSection(
+		const UMovieSceneAudioSection& AudioSection,
+		const FMovieSceneSequenceTransform& SequenceTransform,
+		pxr::UsdPrim& Prim
+	);
+
 	/**
 	 * Functions that call WriterFunc on each UsdTimeSample that corresponds to the keyframes of MovieSceneTrack, writing out time samples for
 	 * attributes of Prim. Mostly used to write out to USD the modified tracks from the automatically generated ULevelSequence owned by
 	 * AUsdStageActors
 	 */
+	USDUTILITIES_API bool ConvertFloatChannel(
+		const FMovieSceneFloatChannel& MovieSceneChannel,
+		const UMovieScene& MovieSceneOuter,
+		const FMovieSceneSequenceTransform& SequenceTransform,
+		const TFunction<void(float, double)>& WriterFunc,
+		UE::FUsdPrim& Prim
+	);
 	USDUTILITIES_API bool ConvertFloatTrack(
 		const UMovieSceneFloatTrack& MovieSceneTrack,
 		const FMovieSceneSequenceTransform& SequenceTransform,
