USD: Fix deadlock when calling PostLoad within UsdAssetCache3.

diff --git a/Engine/Plugins/Runtime/USDCore/Source/USDClasses/Private/USDAssetCache3.cpp b/Engine/Plugins/Runtime/USDCore/Source/USDClasses/Private/USDAssetCache3.cpp
--- a/Engine/Plugins/Runtime/USDCore/Source/USDClasses/Private/USDAssetCache3.cpp
+++ b/Engine/Plugins/Runtime/USDCore/Source/USDClasses/Private/USDAssetCache3.cpp
@@ -1118,16 +1118,20 @@ void UUsdAssetCache3::PostLoad()
 
 	// There is nothing to load us whenever an asset is added to our AssetDirectory while we were unloaded, so let's
 	// make sure we do a new scan whenever we do get loaded to pick up on any new assets that may have been added.
-	// We delay this to the end of the frame though, because we may need to mark ourselves as dirty if we found anything,
+	// We delay this to the next tick though, because we may need to mark ourselves as dirty if we found anything,
 	// and we can't do that within the callstack that calls PostLoad on us.
-	TWeakObjectPtr<UUsdAssetCache3> This(this);
-	AsyncTask(
-		ENamedThreads::GameThread,
-		[This]()
+	//
+	// Note that this was originally within an AsyncTask, but given that RescanAssetDirectory() locks the RWLock, it's
+	// possible to get a deadlock here if the async task is resumed from some unknown point within the callstack of another
+	// asset cache call, so we use the ticker instead.
+	TWeakObjectPtr<UUsdAssetCache3> WeakThis{this};
+	ExecuteOnGameThread(
+		UE_SOURCE_LOCATION,
+		[WeakThis]()
 		{
-			if (This.IsValid())
+			if (UUsdAssetCache3* AssetCache = WeakThis.Get())
 			{
-				This->RescanAssetDirectory();
+				AssetCache->RescanAssetDirectory();
 			}
 		}
 	);
