USD: Fix editor-only code outside WITH_EDITOR blocks missing from cl 34196940

diff --git a/Engine/Plugins/Runtime/USDCore/Source/USDUtilities/Private/USDGeomMeshConversion.cpp b/Engine/Plugins/Runtime/USDCore/Source/USDUtilities/Private/USDGeomMeshConversion.cpp
--- a/Engine/Plugins/Runtime/USDCore/Source/USDUtilities/Private/USDGeomMeshConversion.cpp
+++ b/Engine/Plugins/Runtime/USDCore/Source/USDUtilities/Private/USDGeomMeshConversion.cpp
@@ -3841,8 +3841,10 @@ bool UnrealToUsd::ConvertStaticMesh(
 	// here, so would our referencer and we wouldn't be able to put a transform on it
 	UsdPrim = Stage->DefinePrim(UsdPrim.GetPath(), UnrealToUsd::ConvertToken(bExportMultipleLODs ? TEXT("Xform") : TEXT("Mesh")).Get());
 
+#if WITH_EDITOR
 	const bool bExportNaniteDataAsSourceData = GExportNaniteSourceMeshData && StaticMesh->IsNaniteEnabled()
 											   && StaticMesh->IsHiResMeshDescriptionValid();
+#endif
 
 	bool bExported = false;
 	for (int32 LODIndex = LowestMeshLOD; LODIndex <= HighestMeshLOD; ++LODIndex)
