USD: Translate static mesh, materials and texture assets via Interchange;

Refactor ConvertMaterial into a two-step process that we can use on Interchange;
Support named material slots via Interchange.

diff --git a/Engine/Plugins/Importers/USDImporter/Source/USDSchemas/Private/MeshTranslationImpl.cpp b/Engine/Plugins/Importers/USDImporter/Source/USDSchemas/Private/MeshTranslationImpl.cpp
--- a/Engine/Plugins/Importers/USDImporter/Source/USDSchemas/Private/MeshTranslationImpl.cpp
+++ b/Engine/Plugins/Importers/USDImporter/Source/USDSchemas/Private/MeshTranslationImpl.cpp
@@ -65,9 +65,10 @@ namespace UE::MeshTranslationImplInternal::Private
 		// get the reference UsdPreviewSurface instead, as that is also *its* reference
 		UMaterialInterface* ReferenceMaterial = OneSidedMaterialInstance ? OneSidedMaterialInstance->Parent.Get() : nullptr;
 		UMaterialInterface* ReferenceMaterialTwoSided = nullptr;
-		if (ReferenceMaterial && MeshTranslationImpl::IsReferencePreviewSurfaceMaterial(ReferenceMaterial))
+		if (ReferenceMaterial && UsdUnreal::MaterialUtils::IsReferencePreviewSurfaceMaterial(ReferenceMaterial))
 		{
-			ReferenceMaterialTwoSided = MeshTranslationImpl::GetTwoSidedVersionOfReferencePreviewSurfaceMaterial(ReferenceMaterial);
+			FSoftObjectPath TwoSidedPath = UsdUnreal::MaterialUtils::GetTwoSidedVersionOfReferencePreviewSurfaceMaterial(ReferenceMaterial);
+			ReferenceMaterialTwoSided = Cast<UMaterialInterface>(TwoSidedPath.TryLoad());
 		}
 
 		const FString& DesiredMaterialName = OneSidedMat->GetName() + UnrealIdentifiers::TwoSidedMaterialSuffix;
@@ -444,6 +445,8 @@ TMap<const UsdUtils::FUsdPrimMaterialSlot*, UMaterialInterface*> MeshTranslation
 	bool bShareAssetsForIdenticalPrims
 )
 {
+	using namespace UsdUnreal::MaterialUtils;
+
 	FScopedUnrealAllocs Allocs;
 
 	TMap<const UsdUtils::FUsdPrimMaterialSlot*, UMaterialInterface*> ResolvedMaterials;
@@ -472,16 +475,14 @@ TMap<const UsdUtils::FUsdPrimMaterialSlot*, UMaterialInterface*> MeshTranslation
 			{
 				case UsdUtils::EPrimAssignmentType::DisplayColor:
 				{
-					TOptional<IUsdClassesModule::FDisplayColorMaterial> DisplayColorDesc = IUsdClassesModule::FDisplayColorMaterial::FromString(
-						Slot.MaterialSource
-					);
+					TOptional<FDisplayColorMaterial> DisplayColorDesc = FDisplayColorMaterial::FromString(Slot.MaterialSource);
 
 					if (!DisplayColorDesc.IsSet())
 					{
 						continue;
 					}
 
-					const FSoftObjectPath* ReferencePath = IUsdClassesModule::GetReferenceMaterialPath(DisplayColorDesc.GetValue());
+					const FSoftObjectPath* ReferencePath = GetReferenceMaterialPath(DisplayColorDesc.GetValue());
 					if (!ReferencePath)
 					{
 						continue;
@@ -582,7 +583,7 @@ TMap<const UsdUtils::FUsdPrimMaterialSlot*, UMaterialInterface*> MeshTranslation
 							// material base property overrides at runtime)
 							UMaterialInstance* ExistingInstance = Cast<UMaterialInstance>(ExistingMaterial);
 							const bool bExistingIsInstanceOfReferencePreviewSurface = ExistingInstance
-																					  && MeshTranslationImpl::IsReferencePreviewSurfaceMaterial(
+																					  && UsdUnreal::MaterialUtils::IsReferencePreviewSurfaceMaterial(
 																						  ExistingInstance->Parent
 																					  );
 							if (!OneSidedMat || (!bOneSidedMatIsInstanceOfReferencePreviewSurface && bExistingIsInstanceOfReferencePreviewSurface))
@@ -599,7 +600,10 @@ TMap<const UsdUtils::FUsdPrimMaterialSlot*, UMaterialInterface*> MeshTranslation
 					}
 
 					FString PrefixedMaterialHash = Material ? AssetCache.GetHashForAsset(Material) : FString{};
-					FString HashPrefix = UsdUtils::GetAssetHashPrefix(UsdPrim.GetStage()->GetPrimAtPath(MaterialPrimPath), bShareAssetsForIdenticalPrims);
+					FString HashPrefix = UsdUtils::GetAssetHashPrefix(
+						UsdPrim.GetStage()->GetPrimAtPath(MaterialPrimPath),
+						bShareAssetsForIdenticalPrims
+					);
 
 					// Need to create a two-sided material on-demand, *before* we make it compatible:
 					// This because at runtime we can't just set the base property overrides, and just instead create a new
@@ -993,184 +997,4 @@ void MeshTranslationImpl::RecordSourcePrimsForMaterialSlots(
 	}
 }
 
-UMaterialInterface* MeshTranslationImpl::GetReferencePreviewSurfaceMaterial(EUsdReferenceMaterialProperties ReferenceMaterialProperties)
-{
-	const UUsdProjectSettings* Settings = GetDefault<UUsdProjectSettings>();
-	if (!Settings)
-	{
-		return nullptr;
-	}
-
-	const bool bIsTranslucent = EnumHasAnyFlags(ReferenceMaterialProperties, EUsdReferenceMaterialProperties::Translucent);
-	const bool bIsVT = EnumHasAnyFlags(ReferenceMaterialProperties, EUsdReferenceMaterialProperties::VT);
-	const bool bIsTwoSided = EnumHasAnyFlags(ReferenceMaterialProperties, EUsdReferenceMaterialProperties::TwoSided);
-
-	const FSoftObjectPath* TargetMaterialPath = nullptr;
-	if (bIsTranslucent)
-	{
-		if (bIsVT)
-		{
-			if (bIsTwoSided)
-			{
-				TargetMaterialPath = &Settings->ReferencePreviewSurfaceTranslucentTwoSidedVTMaterial;
-			}
-			else
-			{
-				TargetMaterialPath = &Settings->ReferencePreviewSurfaceTranslucentVTMaterial;
-			}
-		}
-		else
-		{
-			if (bIsTwoSided)
-			{
-				TargetMaterialPath = &Settings->ReferencePreviewSurfaceTranslucentTwoSidedMaterial;
-			}
-			else
-			{
-				TargetMaterialPath = &Settings->ReferencePreviewSurfaceTranslucentMaterial;
-			}
-		}
-	}
-	else
-	{
-		if (bIsVT)
-		{
-			if (bIsTwoSided)
-			{
-				TargetMaterialPath = &Settings->ReferencePreviewSurfaceTwoSidedVTMaterial;
-			}
-			else
-			{
-				TargetMaterialPath = &Settings->ReferencePreviewSurfaceVTMaterial;
-			}
-		}
-		else
-		{
-			if (bIsTwoSided)
-			{
-				TargetMaterialPath = &Settings->ReferencePreviewSurfaceTwoSidedMaterial;
-			}
-			else
-			{
-				TargetMaterialPath = &Settings->ReferencePreviewSurfaceMaterial;
-			}
-		}
-	}
-
-	if (!TargetMaterialPath)
-	{
-		return nullptr;
-	}
-
-	return Cast<UMaterialInterface>(TargetMaterialPath->TryLoad());
-}
-
-UMaterialInterface* MeshTranslationImpl::GetVTVersionOfReferencePreviewSurfaceMaterial(UMaterialInterface* ReferenceMaterial)
-{
-	if (!ReferenceMaterial)
-	{
-		return nullptr;
-	}
-
-	const UUsdProjectSettings* Settings = GetDefault<UUsdProjectSettings>();
-	if (!Settings)
-	{
-		return nullptr;
-	}
-
-	const FSoftObjectPath PathName = ReferenceMaterial->GetPathName();
-	if (PathName.ToString().Contains(TEXT("VT"), ESearchCase::CaseSensitive, ESearchDir::FromEnd))
-	{
-		return ReferenceMaterial;
-	}
-	else if (PathName == Settings->ReferencePreviewSurfaceMaterial)
-	{
-		return Cast<UMaterialInterface>(Settings->ReferencePreviewSurfaceVTMaterial.TryLoad());
-	}
-	else if (PathName == Settings->ReferencePreviewSurfaceTwoSidedMaterial)
-	{
-		return Cast<UMaterialInterface>(Settings->ReferencePreviewSurfaceTwoSidedVTMaterial.TryLoad());
-	}
-	else if (PathName == Settings->ReferencePreviewSurfaceTranslucentMaterial)
-	{
-		return Cast<UMaterialInterface>(Settings->ReferencePreviewSurfaceTranslucentVTMaterial.TryLoad());
-	}
-	else if (PathName == Settings->ReferencePreviewSurfaceTranslucentTwoSidedMaterial)
-	{
-		return Cast<UMaterialInterface>(Settings->ReferencePreviewSurfaceTranslucentTwoSidedVTMaterial.TryLoad());
-	}
-
-	// We should only ever call this function with a ReferenceMaterial that matches one of the above paths
-	ensure(false);
-	return nullptr;
-}
-
-UMaterialInterface* MeshTranslationImpl::GetTwoSidedVersionOfReferencePreviewSurfaceMaterial(UMaterialInterface* ReferenceMaterial)
-{
-	if (!ReferenceMaterial)
-	{
-		return nullptr;
-	}
-
-	const UUsdProjectSettings* Settings = GetDefault<UUsdProjectSettings>();
-	if (!Settings)
-	{
-		return nullptr;
-	}
-
-	const FSoftObjectPath PathName = ReferenceMaterial->GetPathName();
-	if (PathName.ToString().Contains(TEXT("TwoSided"), ESearchCase::CaseSensitive, ESearchDir::FromEnd))
-	{
-		return ReferenceMaterial;
-	}
-	else if (PathName == Settings->ReferencePreviewSurfaceMaterial)
-	{
-		return Cast<UMaterialInterface>(Settings->ReferencePreviewSurfaceTwoSidedMaterial.TryLoad());
-	}
-	else if (PathName == Settings->ReferencePreviewSurfaceTranslucentMaterial)
-	{
-		return Cast<UMaterialInterface>(Settings->ReferencePreviewSurfaceTranslucentTwoSidedMaterial.TryLoad());
-	}
-	else if (PathName == Settings->ReferencePreviewSurfaceVTMaterial)
-	{
-		return Cast<UMaterialInterface>(Settings->ReferencePreviewSurfaceTwoSidedVTMaterial.TryLoad());
-	}
-	else if (PathName == Settings->ReferencePreviewSurfaceTranslucentVTMaterial)
-	{
-		return Cast<UMaterialInterface>(Settings->ReferencePreviewSurfaceTranslucentTwoSidedVTMaterial.TryLoad());
-	}
-
-	// We should only ever call this function with a ReferenceMaterial that matches one of the above paths
-	ensure(false);
-	return nullptr;
-}
-
-bool MeshTranslationImpl::IsReferencePreviewSurfaceMaterial(UMaterialInterface* Material)
-{
-	if (!Material)
-	{
-		return false;
-	}
-
-	const FSoftObjectPath PathName = Material->GetPathName();
-
-	const UUsdProjectSettings* Settings = GetDefault<UUsdProjectSettings>();
-	if (!Settings)
-	{
-		return false;
-	}
-
-	TSet<FSoftObjectPath> ReferenceMaterials = {
-		Settings->ReferencePreviewSurfaceMaterial,
-		Settings->ReferencePreviewSurfaceTranslucentMaterial,
-		Settings->ReferencePreviewSurfaceTwoSidedMaterial,
-		Settings->ReferencePreviewSurfaceTranslucentTwoSidedMaterial,
-		Settings->ReferencePreviewSurfaceVTMaterial,
-		Settings->ReferencePreviewSurfaceTranslucentVTMaterial,
-		Settings->ReferencePreviewSurfaceTwoSidedVTMaterial,
-		Settings->ReferencePreviewSurfaceTranslucentTwoSidedVTMaterial};
-
-	return ReferenceMaterials.Contains(PathName);
-}
-
 #endif	  // #if USE_USD_SDK
diff --git a/Engine/Plugins/Importers/USDImporter/Source/USDSchemas/Private/MeshTranslationImpl.h b/Engine/Plugins/Importers/USDImporter/Source/USDSchemas/Private/MeshTranslationImpl.h
--- a/Engine/Plugins/Importers/USDImporter/Source/USDSchemas/Private/MeshTranslationImpl.h
+++ b/Engine/Plugins/Importers/USDImporter/Source/USDSchemas/Private/MeshTranslationImpl.h
@@ -61,39 +61,6 @@ namespace MeshTranslationImpl
 		const TArray<UsdUtils::FUsdPrimMaterialAssignmentInfo>& LODIndexToMaterialInfo,
 		UUsdMeshAssetUserData* UserData
 	);
-
-	enum class EUsdReferenceMaterialProperties
-	{
-		None = 0,
-		Translucent = 1,
-		VT = 2,
-		TwoSided = 4
-	};
-	ENUM_CLASS_FLAGS(EUsdReferenceMaterialProperties)
-
-	// Returns one of the alternatives of the UsdPreviewSurface reference material depending on the material overrides
-	// provided, and nullptr otherwise
-	UMaterialInterface* GetReferencePreviewSurfaceMaterial(EUsdReferenceMaterialProperties ReferenceMaterialProperties);
-
-	// Returns the VT version of the provided UsdPreviewSurface ReferenceMaterial. Returns the provided ReferenceMaterial back if
-	// it is already a VT-capable reference material, and returns nullptr if ReferenceMaterial isn't one of our reference material
-	// alternatives.
-	// Example: Receives UsdPreviewSurfaceTwoSided -> Returns UsdPreviewSurfaceTwoSidedVT
-	// Example: Receives UsdPreviewSurfaceTwoSidedVT -> Returns UsdPreviewSurfaceTwoSidedVT
-	// Example: Receives SomeOtherReferenceMaterial -> Returns nullptr
-	UMaterialInterface* GetVTVersionOfReferencePreviewSurfaceMaterial(UMaterialInterface* ReferenceMaterial);
-
-	// Returns the two-sided version of the provided UsdPreviewSurface ReferenceMaterial. Returns the provided ReferenceMaterial
-	// back if it is already a two-sided-capable reference material, and returns nullptr if ReferenceMaterial isn't one of our reference
-	// material alternatives.
-	// Example: Receives UsdPreviewSurfaceTranslucent -> Returns UsdPreviewSurfaceTwoSidedTranslucent
-	// Example: Receives UsdPreviewSurfaceTwoSidedTranslucent -> Returns UsdPreviewSurfaceTwoSidedTranslucent
-	// Example: Receives SomeOtherReferenceMaterial -> Returns nullptr
-	UMaterialInterface* GetTwoSidedVersionOfReferencePreviewSurfaceMaterial(UMaterialInterface* ReferenceMaterial);
-
-	// Returns whether the Material is one of the UsdPreviewSurface reference materials (which can be reassigned by the
-	// user on a per project basis)
-	bool IsReferencePreviewSurfaceMaterial(UMaterialInterface* Material);
 }	 // namespace MeshTranslationImpl
 
 #endif	  // #if USE_USD_SDK
diff --git a/Engine/Plugins/Importers/USDImporter/Source/USDSchemas/Private/USDGeomMeshTranslator.cpp b/Engine/Plugins/Importers/USDImporter/Source/USDSchemas/Private/USDGeomMeshTranslator.cpp
--- a/Engine/Plugins/Importers/USDImporter/Source/USDSchemas/Private/USDGeomMeshTranslator.cpp
+++ b/Engine/Plugins/Importers/USDImporter/Source/USDSchemas/Private/USDGeomMeshTranslator.cpp
@@ -61,21 +61,6 @@
 #include "pxr/usd/usdSkel/bindingAPI.h"
 #include "USDIncludesEnd.h"
 
-static float GMeshNormalRepairThreshold = 0.05f;
-static FAutoConsoleVariableRef CVarMeshNormalRepairThreshold(
-	TEXT("USD.MeshNormalRepairThreshold"),
-	GMeshNormalRepairThreshold,
-	TEXT("We will try repairing up to this fraction of a Mesh's normals when invalid. If a Mesh has more invalid normals than this, we will "
-		 "recompute all of them. Defaults to 0.05 (5% of all normals).")
-);
-
-static bool GSkipMeshTangentComputation = false;
-static FAutoConsoleVariableRef CVarSkipMeshTangentComputation(
-	TEXT("USD.SkipMeshTangentComputation"),
-	GSkipMeshTangentComputation,
-	TEXT("Skip computing tangents for meshes. With meshes with a huge numer of vertices, it can take a very long time to compute them.")
-);
-
 static bool GBuildReversedIndexBuffer = false;
 static FAutoConsoleVariableRef CVarBuildReversedIndexBuffer(
 	TEXT("USD.StaticMesh.BuildSettings.BuildReversedIndexBuffer"),
@@ -589,74 +574,6 @@ namespace UsdGeomMeshTranslatorImpl
 		}
 	}
 
-	void RepairNormalsAndTangents(const FString& PrimPath, FMeshDescription& MeshDescription)
-	{
-		FStaticMeshConstAttributes Attributes{MeshDescription};
-		TArrayView<const FVector3f> VertexInstanceNormals = Attributes.GetVertexInstanceNormals().GetRawArray();
-
-		// Similar to FStaticMeshOperations::AreNormalsAndTangentsValid but we don't care about tangents since we never
-		// read those from USD
-		uint64 InvalidNormalCount = 0;
-		for (const FVertexInstanceID VertexInstanceID : MeshDescription.VertexInstances().GetElementIDs())
-		{
-			if (VertexInstanceNormals[VertexInstanceID].IsNearlyZero() || VertexInstanceNormals[VertexInstanceID].ContainsNaN())
-			{
-				++InvalidNormalCount;
-			}
-		}
-		if (InvalidNormalCount == 0)
-		{
-			return;
-		}
-
-		const float InvalidNormalFraction = (float)InvalidNormalCount / (float)VertexInstanceNormals.Num();
-
-		// We always need to do this at this point as ComputeTangentsAndNormals will end up computing tangents anyway
-		// and our triangle tangents are always invalid
-		FStaticMeshOperations::ComputeTriangleTangentsAndNormals(MeshDescription);
-
-		const static FString MeshNormalRepairThresholdText = TEXT("USD.MeshNormalRepairThreshold");
-
-		// Make sure our normals can be rebuilt from MeshDescription::InitializeAutoGeneratedAttributes in case some tool needs them.
-		// Always force-compute tangents here as we never have them anyway. If we don't force them to be recomputed we'll get
-		// the worst of both worlds as some of these will be arbitrarily recomputed anyway, and some will be left invalid
-		EComputeNTBsFlags Options = GSkipMeshTangentComputation ? EComputeNTBsFlags::None
-																: EComputeNTBsFlags::UseMikkTSpace | EComputeNTBsFlags::Tangents;
-
-		// Repairing can take a long time for degenerate triangles (UE-194839)
-		Options |= EComputeNTBsFlags::IgnoreDegenerateTriangles;
-
-		if (InvalidNormalFraction >= GMeshNormalRepairThreshold)
-		{
-			Options |= EComputeNTBsFlags::Normals;
-			UE_LOG(
-				LogUsd,
-				Verbose,
-				TEXT("%f%% of the normals from Mesh prim '%s' are invalid or unusable. This is at or above the threshold of '%f%%' (configurable via "
-					 "the cvar '%s'), so normals will be discarded and fully recomputed. Note that when the cvar "
-					 "'USD.Subdiv.IgnoreNormalsWhenSubdividing' is true it is expected for subdivision meshes to have their normals discarded."),
-				InvalidNormalFraction * 100.0f,
-				*PrimPath,
-				GMeshNormalRepairThreshold * 100.0f,
-				*MeshNormalRepairThresholdText
-			);
-		}
-		else if (InvalidNormalFraction > 0)
-		{
-			UE_LOG(
-				LogUsd,
-				Verbose,
-				TEXT("%f%% of the normals from Mesh prim '%s' are invalid or unusable. This is below the threshold of '%f%%' (configurable via the "
-					 "cvar '%s'), so the invalid normals will be repaired."),
-				InvalidNormalFraction * 100.0f,
-				*PrimPath,
-				GMeshNormalRepairThreshold * 100.0f,
-				*MeshNormalRepairThresholdText
-			);
-		}
-		FStaticMeshOperations::ComputeTangentsAndNormals(MeshDescription, Options);
-	}
-
 	UStaticMesh* CreateStaticMesh(
 		const UE::FUsdPrim& Prim,
 		TArray<FMeshDescription>& LODIndexToMeshDescription,
@@ -702,8 +619,13 @@ namespace UsdGeomMeshTranslatorImpl
 			// could be tricky to do since some of these build steps are async/thread-pool based.
 			SHA1.Update(reinterpret_cast<const uint8*>(&bShouldEnableNanite), sizeof(bShouldEnableNanite));
 
-			// Hash the threshhold so that if we update it and reload we'll regenerate static meshes
-			SHA1.Update(reinterpret_cast<const uint8*>(&GMeshNormalRepairThreshold), sizeof(GMeshNormalRepairThreshold));
+			// Hash these so that if we update them and reload we'll regenerate static meshes
+			static IConsoleVariable* NormalThresholdCvar = IConsoleManager::Get().FindConsoleVariable(TEXT("USD.MeshNormalRepairThreshold"));
+			static IConsoleVariable* SkipTangentComputationCvar = IConsoleManager::Get().FindConsoleVariable(TEXT("USD.SkipMeshTangentComputation"));
+			const float MeshNormalRepairThreshold = NormalThresholdCvar ? NormalThresholdCvar->GetFloat() : 1.0f;
+			const bool SkipTangentComputation = SkipTangentComputationCvar && SkipTangentComputationCvar->GetBool();
+			SHA1.Update(reinterpret_cast<const uint8*>(&MeshNormalRepairThreshold), sizeof(MeshNormalRepairThreshold));
+			SHA1.Update(reinterpret_cast<const uint8*>(&SkipTangentComputation), sizeof(SkipTangentComputation));
 
 			SHA1.Final();
 			SHA1.GetHash(&AllLODHash.Hash[0]);
@@ -733,7 +655,7 @@ namespace UsdGeomMeshTranslatorImpl
 			{
 				FMeshDescription& MeshDescription = LODIndexToMeshDescription[LODIndex];
 
-				RepairNormalsAndTangents(MeshName, MeshDescription);
+				UsdUtils::RepairNormalsAndTangents(MeshName, MeshDescription);
 
 				FStaticMeshSourceModel& SourceModel = StaticMesh->AddSourceModel();
 				SourceModel.BuildSettings.bGenerateLightmapUVs = false;
diff --git a/Engine/Plugins/Importers/USDImporter/Source/USDSchemas/Private/USDShadeMaterialTranslator.cpp b/Engine/Plugins/Importers/USDImporter/Source/USDSchemas/Private/USDShadeMaterialTranslator.cpp
--- a/Engine/Plugins/Importers/USDImporter/Source/USDSchemas/Private/USDShadeMaterialTranslator.cpp
+++ b/Engine/Plugins/Importers/USDImporter/Source/USDSchemas/Private/USDShadeMaterialTranslator.cpp
@@ -78,7 +78,8 @@ namespace UE::UsdShadeTranslator::Private
 						// Important to not use GetBaseMaterial() here because if our parent is the translucent we'll
 						// get the reference UsdPreviewSurface instead, as that is also *its* reference
 						UMaterialInterface* ReferenceMaterial = MaterialInstance->Parent.Get();
-						UMaterialInterface* ReferenceMaterialVT = MeshTranslationImpl::GetVTVersionOfReferencePreviewSurfaceMaterial(ReferenceMaterial
+						UMaterialInterface* ReferenceMaterialVT = Cast<UMaterialInterface>(
+							UsdUnreal::MaterialUtils::GetVTVersionOfReferencePreviewSurfaceMaterial(ReferenceMaterial).TryLoad()
 						);
 						if (ReferenceMaterial == ReferenceMaterialVT)
 						{
@@ -406,16 +407,18 @@ void FUsdShadeMaterialTranslator::CreateAssets()
 				UE::UsdShadeTranslator::Private::UpgradeMaterialsAndTexturesToVT(NonVTTextures, Context);
 			}
 
-			MeshTranslationImpl::EUsdReferenceMaterialProperties Properties = MeshTranslationImpl::EUsdReferenceMaterialProperties::None;
+			EUsdReferenceMaterialProperties Properties = EUsdReferenceMaterialProperties::None;
 			if (bIsMaterialTranslucent)
 			{
-				Properties |= MeshTranslationImpl::EUsdReferenceMaterialProperties::Translucent;
+				Properties |= EUsdReferenceMaterialProperties::Translucent;
 			}
 			if (VTTextures.Num() > 0)
 			{
-				Properties |= MeshTranslationImpl::EUsdReferenceMaterialProperties::VT;
+				Properties |= EUsdReferenceMaterialProperties::VT;
 			}
-			UMaterialInterface* ReferenceMaterial = MeshTranslationImpl::GetReferencePreviewSurfaceMaterial(Properties);
+			UMaterialInterface* ReferenceMaterial = Cast<UMaterialInterface>(
+				UsdUnreal::MaterialUtils::GetReferencePreviewSurfaceMaterial(Properties).TryLoad()
+			);
 
 			if (ensure(ReferenceMaterial))
 			{
@@ -454,12 +457,14 @@ void FUsdShadeMaterialTranslator::CreateAssets()
 			Context->ObjectFlags | RF_Transient,	// We never want MIDs to become assets in the content browser
 			[bIsMaterialTranslucent](UPackage* Outer, FName SanitizedName, EObjectFlags FlagsToUse)
 			{
-				MeshTranslationImpl::EUsdReferenceMaterialProperties Properties = MeshTranslationImpl::EUsdReferenceMaterialProperties::None;
+				EUsdReferenceMaterialProperties Properties = EUsdReferenceMaterialProperties::None;
 				if (bIsMaterialTranslucent)
 				{
-					Properties |= MeshTranslationImpl::EUsdReferenceMaterialProperties::Translucent;
+					Properties |= EUsdReferenceMaterialProperties::Translucent;
 				}
-				UMaterialInterface* ReferenceMaterial = MeshTranslationImpl::GetReferencePreviewSurfaceMaterial(Properties);
+				UMaterialInterface* ReferenceMaterial = Cast<UMaterialInterface>(
+					UsdUnreal::MaterialUtils::GetReferencePreviewSurfaceMaterial(Properties).TryLoad()
+				);
 
 				UMaterialInstanceDynamic* CreatedMID = UMaterialInstanceDynamic::Create(ReferenceMaterial, Outer, SanitizedName);
 				CreatedMID->ClearFlags(CreatedMID->GetFlags());
diff --git a/Engine/Plugins/Interchange/Runtime/Source/Import/InterchangeImport.Build.cs b/Engine/Plugins/Interchange/Runtime/Source/Import/InterchangeImport.Build.cs
--- a/Engine/Plugins/Interchange/Runtime/Source/Import/InterchangeImport.Build.cs
+++ b/Engine/Plugins/Interchange/Runtime/Source/Import/InterchangeImport.Build.cs
@@ -44,6 +44,7 @@ namespace UnrealBuildTool.Rules
 					"RHI",
 					"TextureUtilitiesCommon",
 					"UnrealUSDWrapper",
+					"USDClasses",
 					"USDUtilities",
 					"VariantManagerContent",
 				}
diff --git a/Engine/Plugins/Interchange/Runtime/Source/Import/Private/Usd/InterchangeUsdTranslator.cpp b/Engine/Plugins/Interchange/Runtime/Source/Import/Private/Usd/InterchangeUsdTranslator.cpp
--- a/Engine/Plugins/Interchange/Runtime/Source/Import/Private/Usd/InterchangeUsdTranslator.cpp
+++ b/Engine/Plugins/Interchange/Runtime/Source/Import/Private/Usd/InterchangeUsdTranslator.cpp
@@ -1,26 +1,49 @@
-// Copyright Epic Games, Inc. All Rights Reserved. 
+// Copyright Epic Games, Inc. All Rights Reserved.
 
 #include "Usd/InterchangeUsdTranslator.h"
 
 #include "UnrealUSDWrapper.h"
 #include "UsdWrappers/SdfPath.h"
-#include "UsdWrappers/UsdGeomXformable.h"
 #include "UsdWrappers/UsdAttribute.h"
+#include "UsdWrappers/UsdGeomXformable.h"
 #include "UsdWrappers/UsdPrim.h"
 #include "UsdWrappers/UsdStage.h"
 #include "UsdWrappers/UsdTyped.h"
 #include "UsdWrappers/VtValue.h"
 
 #include "USDConversionUtils.h"
+#include "USDGeomMeshConversion.h"
+#include "USDLog.h"
+#include "USDMaterialUtils.h"
 #include "USDPrimConversion.h"
+#include "USDShadeConversion.h"
+#include "USDStageOptions.h"
+#include "USDTypesConversion.h"
 #include "USDValueConversion.h"
 
+#include "Async/Async.h"
 #include "HAL/IConsoleManager.h"
 #include "InterchangeCameraNode.h"
 #include "InterchangeImportLog.h"
 #include "InterchangeLightNode.h"
+#include "InterchangeMaterialInstanceNode.h"
 #include "InterchangeMeshNode.h"
 #include "InterchangeSceneNode.h"
+#include "InterchangeTexture2DNode.h"
+#include "InterchangeTranslatorHelper.h"
+#include "Internationalization/Regex.h"
+#include "StaticMeshAttributes.h"
+#include "StaticMeshOperations.h"
+#include "UDIMUtilities.h"
+
+#if USE_USD_SDK
+#include "USDIncludesStart.h"
+#include "pxr/usd/usdGeom/tokens.h"
+#include "pxr/usd/usdPhysics/tokens.h"
+#include "pxr/usd/usdShade/tokens.h"
+#include "USDIncludesEnd.h"
+
+#endif	  // USE_USD_SDK
 
 #include UE_INLINE_GENERATED_CPP_BY_NAME(InterchangeUsdTranslator)
 
@@ -33,7 +56,49 @@ static FAutoConsoleVariableRef CVarInterchangeEnableUSDImport(
 	TEXT("Whether USD support is enabled.")
 );
 
+namespace UE::InterchangeUsdTranslator::Private
+{
+	const static FString CameraPrefix = TEXT("\\Camera\\");
+	const static FString LightPrefix = TEXT("\\Light\\");
+	const static FString MaterialPrefix = TEXT("\\Material\\");
+	const static FString MeshPrefix = TEXT("\\Mesh\\");
+
+	// Small container that we can use Pimpl with so we don't have to include too many USD includes on the header file.
+	//
+	// It also skirts around a small complication where UInterchangeUSDTranslator::Translate is const, and yet we must
+	// keep and modify some members (like UsdStage) for when the payload functions get called later... The other translators
+	// use mutable or const casts, but with the Impl we don't need to!
+	class UInterchangeUSDTranslatorImpl
+	{
+	public:
+		// We have to keep a stage reference so that we can parse the payloads after Translate() complete.
+		// ReleaseSource() clears this member, once translation is complete.
+		UE::FUsdStage UsdStage;
+
+		TObjectPtr<UInterchangeUsdTranslatorSettings> TranslatorSettings = nullptr;
+
+#if USE_USD_SDK
+		// On UInterchangeUSDTranslator::Translate we set this up based on our TranslatorSettings, and then
+		// we can reuse it (otherwise we have to keep converting the FNames into Tokens all the time)
+		UsdToUnreal::FUsdMeshConversionOptions CachedMeshConversionOptions;
+#endif
+	};
+}
+
+UInterchangeUsdTranslatorSettings::UInterchangeUsdTranslatorSettings()
+	: GeometryPurpose((int32)(EUsdPurpose::Default | EUsdPurpose::Proxy | EUsdPurpose::Render | EUsdPurpose::Guide))
+	, RenderContext(TEXT("unreal"))	   // The proper definition of this is on the USDSchemas module, which we can't depend on
+	, MaterialPurpose(*UnrealIdentifiers::MaterialPreviewPurpose)
+	, bOverrideStageOptions(false)
+	, StageOptions{
+		  0.01f,			   // MetersPerUnit
+		  EUsdUpAxis::ZAxis	   // UpAxis
+	  }
+{
+}
+
 UInterchangeUSDTranslator::UInterchangeUSDTranslator()
+	: Impl(MakeUnique<UE::InterchangeUsdTranslator::Private::UInterchangeUSDTranslatorImpl>())
 {
 }
 
@@ -59,26 +124,331 @@ TArray<FString> UInterchangeUSDTranslator::GetSupportedFormats() const
 
 namespace UE::InterchangeUsdTranslator::Private
 {
-	void AddLightNode(const UE::FUsdPrim& Prim, UInterchangeBaseNodeContainer& NodeContainer)
+#if USE_USD_SDK
+	FString EncodeTexturePayloadKey(const UsdToUnreal::FTextureParameterValue& Value)
+	{
+		// Encode the compression settings onto the payload key as we need to move that into the
+		// payload data within UInterchangeUSDTranslator::GetTexturePayloadData.
+		//
+		// This should be a temporary thing, and in the future we'll be able to store compression
+		// settings directly on the texture translated node
+		return Value.TextureFilePath + TEXT("\\") + LexToString((int32)Value.Group);
+	}
+
+	bool DecodeTexturePayloadKey(const FString& PayloadKey, FString& OutTextureFilePath, TextureGroup& OutTextureGroup)
+	{
+		// Use split from end here so that we ignore any backslashes within the file path itself
+		FString FilePath;
+		FString TextureGroupStr;
+		bool bSplit = PayloadKey.Split(TEXT("\\"), &FilePath, &TextureGroupStr, ESearchCase::CaseSensitive, ESearchDir::FromEnd);
+		if (!bSplit)
+		{
+			return false;
+		}
+
+		OutTextureFilePath = FilePath;
+
+		int32 TempInt;
+		if (LexTryParseString<int32>(TempInt, *TextureGroupStr))
+		{
+			OutTextureGroup = (TextureGroup)TempInt;
+		}
+
+		return true;
+	}
+
+	void AddTextureNode(
+		const UE::FUsdPrim& Prim,
+		const FString& NodeUid,
+		const UsdToUnreal::FTextureParameterValue& Value,
+		UInterchangeBaseNodeContainer& NodeContainer
+	)
+	{
+		FString PrimPath = Prim.GetPrimPath().GetString();
+		FString NodeName{FPaths::GetCleanFilename(Value.TextureFilePath)};
+
+		// Check if Node already exist with this ID
+		if (const UInterchangeTexture2DNode* Node = Cast<UInterchangeTexture2DNode>(NodeContainer.GetNode(NodeUid)))
+		{
+			return;
+		}
+
+		UInterchangeTexture2DNode* Node = NewObject<UInterchangeTexture2DNode>(&NodeContainer);
+		Node->InitializeNode(NodeUid, NodeName, EInterchangeNodeContainerType::TranslatedAsset);
+		Node->SetPayLoadKey(EncodeTexturePayloadKey(Value));
+
+		static_assert((int)TextureAddress::TA_Wrap == (int)EInterchangeTextureWrapMode::Wrap);
+		static_assert((int)TextureAddress::TA_Clamp == (int)EInterchangeTextureWrapMode::Clamp);
+		static_assert((int)TextureAddress::TA_Mirror == (int)EInterchangeTextureWrapMode::Mirror);
+		Node->SetCustomWrapU((EInterchangeTextureWrapMode)Value.AddressX);
+		Node->SetCustomWrapV((EInterchangeTextureWrapMode)Value.AddressY);
+
+		Node->SetCustomSRGB(Value.GetSRGBValue());
+
+		// Provide the other UDIM tiles
+		//
+		// Note: There is an bImportUDIM option on UInterchangeGenericTexturePipeline that is exclusively used within
+		// UInterchangeGenericTexturePipeline::HandleCreationOfTextureFactoryNode in order to essentially do the exact same
+		// thing as we do here. In theory, we shouldn't need to do this then, and in fact it is a bit bad to do so because
+		// we will always parse these UDIMs whether the option is enabled or disabled. The issue however is that (as of the
+		// time of this writing) UInterchangeGenericTexturePipeline::HandleCreationOfTextureFactoryNode is hard-coded to expect
+		// the texture payload key to be just the texture file path. We can't do that, because we need to also encode
+		// the texture compression settings onto payload key...
+		//
+		// All of that is to say that everything will actually work fine, but if you uncheck "bImportUDIM" on the import options
+		// you will still get UDIMs (for now).
+		if (Value.bIsUDIM)
+		{
+			TMap<int32, FString> TileIndexToPath = UE::TextureUtilitiesCommon::GetUDIMBlocksFromSourceFile(
+				Value.TextureFilePath,
+				UE::TextureUtilitiesCommon::DefaultUdimRegexPattern
+			);
+			Node->SetSourceBlocks(MoveTemp(TileIndexToPath));
+		}
+
+		NodeContainer.AddNode(Node);
+	}
+
+	// We use this visitor to set UsdToUnreal::FParameterValue TVariant values onto UInterchangeMaterialInstanceNodes
+	struct FParameterValueVisitor
+	{
+		FParameterValueVisitor(
+			const UE::FUsdPrim& InPrim,
+			UInterchangeBaseNodeContainer& InNodeContainer,
+			UInterchangeMaterialInstanceNode& InMaterialNode,
+			const TMap<FString, int32>& InPrimvarToUVIndex
+		)
+			: Prim(InPrim)
+			, NodeContainer(InNodeContainer)
+			, MaterialNode(InMaterialNode)
+			, PrimvarToUVIndex(InPrimvarToUVIndex)
+		{
+		}
+
+		void operator()(const float Value) const
+		{
+			MaterialNode.AddScalarParameterValue(*ParameterName, Value);
+
+			// Disable the texture input since we have a direct value
+			MaterialNode.AddScalarParameterValue(*FString::Printf(TEXT("Use%sTexture"), **ParameterName), 0.0f);
+		}
+
+		void operator()(const FVector& Value) const
+		{
+			MaterialNode.AddVectorParameterValue(*ParameterName, FLinearColor{Value});
+
+			// Disable the texture input since we have a direct value
+			MaterialNode.AddScalarParameterValue(*FString::Printf(TEXT("Use%sTexture"), **ParameterName), 0.0f);
+		}
+
+		void operator()(const UsdToUnreal::FTextureParameterValue& Value) const
+		{
+			// Emit texture node itself (this is the main place where this happens)
+			// Note that the node name isn't just the texture path, as we may have multiple material users of this texture
+			// with different settings, and so we need separate translated nodes for each material and parameter
+			FString TextureUid = FString::Printf(TEXT("Texture:%s:%s"), *Prim.GetPrimPath().GetString(), **ParameterName);
+			AddTextureNode(Prim, TextureUid, Value, NodeContainer);
+
+			// Actual texture assignment
+			MaterialNode.AddTextureParameterValue(*FString::Printf(TEXT("%sTexture"), **ParameterName), TextureUid);
+			MaterialNode.AddScalarParameterValue(*FString::Printf(TEXT("Use%sTexture"), **ParameterName), 1.0f);
+
+			// UV transform
+			FLinearColor ScaleAndTranslation = FLinearColor{
+				Value.UVScale.GetVector()[0],
+				Value.UVScale.GetVector()[1],
+				Value.UVTranslation[0],
+				Value.UVTranslation[1]};
+			MaterialNode.AddVectorParameterValue(*FString::Printf(TEXT("%sScaleTranslation"), **ParameterName), ScaleAndTranslation);
+			MaterialNode.AddScalarParameterValue(*FString::Printf(TEXT("%sRotation"), **ParameterName), Value.UVRotation);
+
+			// UV index
+			if (const int32* FoundIndex = PrimvarToUVIndex.Find(Value.Primvar))
+			{
+				MaterialNode.AddScalarParameterValue(*FString::Printf(TEXT("%sUVIndex"), **ParameterName), *FoundIndex);
+			}
+			else
+			{
+				UE_LOG(
+					LogUsd,
+					Warning,
+					TEXT("Failed to find primvar '%s' when setting material parameter '%s' on material '%s'. Available primvars and UV "
+						 "indices: %s.%s"),
+					*Value.Primvar,
+					**ParameterName,
+					*Prim.GetPrimPath().GetString(),
+					*UsdUtils::StringifyMap(PrimvarToUVIndex),
+					Value.Primvar.IsEmpty() ? TEXT(
+						" Is your UsdUVTexture Shader missing the 'inputs:st' attribute? (It specifies which UV set to sample the texture with)"
+					)
+											: TEXT("")
+				);
+			}
+
+			// Component mask (which channel of the texture to use)
+			FLinearColor ComponentMask = FLinearColor::Black;
+			switch (Value.OutputIndex)
+			{
+				case 0:	   // RGB
+					ComponentMask = FLinearColor{1.f, 1.f, 1.f, 0.f};
+					break;
+				case 1:	   // R
+					ComponentMask = FLinearColor{1.f, 0.f, 0.f, 0.f};
+					break;
+				case 2:	   // G
+					ComponentMask = FLinearColor{0.f, 1.f, 0.f, 0.f};
+					break;
+				case 3:	   // B
+					ComponentMask = FLinearColor{0.f, 0.f, 1.f, 0.f};
+					break;
+				case 4:	   // A
+					ComponentMask = FLinearColor{0.f, 0.f, 0.f, 1.f};
+					break;
+			}
+			MaterialNode.AddVectorParameterValue(*FString::Printf(TEXT("%sTextureComponent"), **ParameterName), ComponentMask);
+		}
+
+		void operator()(const UsdToUnreal::FPrimvarReaderParameterValue& Value) const
+		{
+			MaterialNode.AddVectorParameterValue(*ParameterName, FLinearColor{Value.FallbackValue});
+
+			if (Value.PrimvarName == TEXT("displayColor"))
+			{
+				MaterialNode.AddScalarParameterValue(TEXT("UseVertexColorForBaseColor"), 1.0f);
+			}
+		}
+
+		void operator()(const bool Value) const
+		{
+			MaterialNode.AddScalarParameterValue(*ParameterName, static_cast<float>(Value));
+		}
+
+	public:
+		const UE::FUsdPrim& Prim;
+		UInterchangeBaseNodeContainer& NodeContainer;
+		UInterchangeMaterialInstanceNode& MaterialNode;
+		const TMap<FString, int32>& PrimvarToUVIndex;
+		const FString* ParameterName = nullptr;
+	};
+
+	void AddMaterialInstanceNode(
+		const UE::FUsdPrim& Prim,
+		UInterchangeUSDTranslatorImpl* TranslatorImpl,
+		UInterchangeBaseNodeContainer& NodeContainer
+	)
 	{
-		FString NodeUid = TEXT("Light:") + Prim.GetPrimPath().GetString();
+		FString PrimPath = Prim.GetPrimPath().GetString();
+		FString NodeUid = MaterialPrefix + PrimPath;
 		FString NodeName(Prim.GetName().ToString());
 
-		auto Attributes = Prim.GetAttributes();
-		for (auto& Attribute : Attributes)
+		// Check if Node already exist with this ID
+		if (const UInterchangeMaterialInstanceNode* Node = Cast<UInterchangeMaterialInstanceNode>(NodeContainer.GetNode(NodeUid)))
+		{
+			return;
+		}
+
+		UInterchangeMaterialInstanceNode* MaterialNode = NewObject<UInterchangeMaterialInstanceNode>(&NodeContainer);
+		MaterialNode->InitializeNode(NodeUid, NodeName, EInterchangeNodeContainerType::TranslatedAsset);
+		MaterialNode->SetAssetName(NodeName);
+		NodeContainer.AddNode(MaterialNode);
+
+		UsdToUnreal::FUsdPreviewSurfaceMaterialData MaterialData;
+		FString RenderContext = TranslatorImpl->TranslatorSettings ? TranslatorImpl->TranslatorSettings->RenderContext.ToString() : FString();
+		const bool bSuccess = UsdToUnreal::ConvertMaterial(Prim, MaterialData, TranslatorImpl->TranslatorSettings ? *RenderContext : nullptr);
+
+		// Set all the parameter values to the interchange node
+		bool bHasUDIMTexture = false;
+		FParameterValueVisitor Visitor{Prim, NodeContainer, *MaterialNode, MaterialData.PrimvarToUVIndex};
+		for (TPair<FString, UsdToUnreal::FParameterValue>& Pair : MaterialData.Parameters)
+		{
+			Visitor.ParameterName = &Pair.Key;
+			Visit(Visitor, Pair.Value);
+
+			// Also simultaneously check if any of these parameters wants to be an UDIM texture so that we can use the VT reference material later
+			if (!bHasUDIMTexture)
+			{
+				if (UsdToUnreal::FTextureParameterValue* TextureParameter = Pair.Value.TryGet<UsdToUnreal::FTextureParameterValue>())
+				{
+					if (TextureParameter->bIsUDIM)
+					{
+						bHasUDIMTexture = true;
+					}
+				}
+			}
+		}
+
+		// Find and set the right reference material
+		//
+		// TODO: Proper VT texture support (we'd need to know the texture resolution at this point, and we haven't parsed them yet...).
+		// The way it currently works on Interchange is that the factory will create a VT or nonVT version of the texture to match the
+		// material parameter slot. Since we'll currently never set the VT reference material, it essentially means it will always
+		// downgrade our VT textures to non-VT.
+		// The only exception is how we upgrade the reference material to VT in case we have any UDIM textures, as those are trivial to
+		// check for (we don't have to actually load the textures to do it)
+		EUsdReferenceMaterialProperties Properties = EUsdReferenceMaterialProperties::None;
+		if (UsdUtils::IsMaterialTranslucent(MaterialData))
+		{
+			Properties |= EUsdReferenceMaterialProperties::Translucent;
+		}
+		if (bHasUDIMTexture)
+		{
+			Properties |= EUsdReferenceMaterialProperties::VT;
+		}
+
+		FSoftObjectPath ReferenceMaterial = UsdUnreal::MaterialUtils::GetReferencePreviewSurfaceMaterial(Properties);
+		MaterialNode->SetCustomParent(ReferenceMaterial.ToString());
+	}
+
+	void AddDisplayColorMaterialInstanceNodeIfNeeded(UInterchangeBaseNodeContainer& NodeContainer, const FString& DisplayColorDesc)
+	{
+		using namespace UsdUnreal::MaterialUtils;
+
+		FString NodeUid = MaterialPrefix + DisplayColorDesc;
+
+		// We'll treat the DisplayColorDesc (something like "!DisplayColor_1_0") as the material instance UID here
+		const UInterchangeMaterialInstanceNode* Node = Cast<UInterchangeMaterialInstanceNode>(NodeContainer.GetNode(DisplayColorDesc));
+		if (Node)
+		{
+			return;
+		}
+
+		// Need to create a new instance
+		TOptional<FDisplayColorMaterial> ParsedMat = FDisplayColorMaterial::FromString(DisplayColorDesc);
+		if (!ParsedMat)
+		{
+			return;
+		}
+		FString NodeName = ParsedMat->ToPrettyString();
+
+		const FSoftObjectPath* ReferenceMaterialPath = GetReferenceMaterialPath(ParsedMat.GetValue());
+		if (!ReferenceMaterialPath)
 		{
-			UE_LOG(LogInterchangeImport, Log, TEXT("%s %s %s"), *Attribute.GetName().ToString(), *Attribute.GetBaseName().ToString(), *Attribute.GetTypeName().ToString());
+			return;
 		}
 
+		// Not needed
+		const FString ParentNodeUid;
+		UInterchangeMaterialInstanceNode* NewNode = UInterchangeMaterialInstanceNode::Create(&NodeContainer, DisplayColorDesc, ParentNodeUid);
+		NewNode->InitializeNode(NodeUid, NodeName, EInterchangeNodeContainerType::TranslatedAsset);
+
+		NewNode->SetCustomParent(ReferenceMaterialPath->GetAssetPathString());
+	}
+
+	void AddLightNode(const UE::FUsdPrim& Prim, UInterchangeBaseNodeContainer& NodeContainer)
+	{
+		FString NodeUid = LightPrefix + Prim.GetPrimPath().GetString();
+		FString NodeName(Prim.GetName().ToString());
+
 		// #ueent_todo: Retrieve light attributes and set them below
 
 		if (Prim.IsA(TEXT("DistantLight")))
 		{
 			UInterchangeDirectionalLightNode* LightNode = NewObject<UInterchangeDirectionalLightNode>(&NodeContainer);
 			LightNode->InitializeNode(NodeUid, NodeName, EInterchangeNodeContainerType::TranslatedAsset);
+			LightNode->SetAssetName(NodeName);
 
-			//LightNode->SetCustomLightColor(Color);
-			//LightNode->SetCustomIntensity(Intensity);
+			// LightNode->SetCustomLightColor(Color);
+			// LightNode->SetCustomIntensity(Intensity);
 
 			NodeContainer.AddNode(LightNode);
 		}
@@ -88,13 +458,14 @@ namespace UE::InterchangeUsdTranslator::Private
 			{
 				UInterchangeSpotLightNode* LightNode = NewObject<UInterchangeSpotLightNode>(&NodeContainer);
 				LightNode->InitializeNode(NodeUid, NodeName, EInterchangeNodeContainerType::TranslatedAsset);
+				LightNode->SetAssetName(NodeName);
 
 				LightNode->SetCustomIntensityUnits(EInterchangeLightUnits::Lumens);
-				//LightNode->SetCustomLightColor(Color));
-				//LightNode->SetCustomIntensity(Intensity);
+				// LightNode->SetCustomLightColor(Color));
+				// LightNode->SetCustomIntensity(Intensity);
 
-				//LightNode->SetCustomInnerConeAngle(FMath::RadiansToDegrees(InnerConeAngle));
-				//LightNode->SetCustomOuterConeAngle(FMath::RadiansToDegrees(OuterConeAngle));
+				// LightNode->SetCustomInnerConeAngle(FMath::RadiansToDegrees(InnerConeAngle));
+				// LightNode->SetCustomOuterConeAngle(FMath::RadiansToDegrees(OuterConeAngle));
 
 				NodeContainer.AddNode(LightNode);
 			}
@@ -102,12 +473,13 @@ namespace UE::InterchangeUsdTranslator::Private
 			{
 				UInterchangePointLightNode* LightNode = NewObject<UInterchangePointLightNode>(&NodeContainer);
 				LightNode->InitializeNode(NodeUid, NodeName, EInterchangeNodeContainerType::TranslatedAsset);
+				LightNode->SetAssetName(NodeName);
 
 				LightNode->SetCustomIntensityUnits(EInterchangeLightUnits::Lumens);
-				//LightNode->SetCustomLightColor(FLinearColor(Color));
-				//LightNode->SetCustomIntensity(Intensity);
+				// LightNode->SetCustomLightColor(FLinearColor(Color));
+				// LightNode->SetCustomIntensity(Intensity);
 
-				//LightNode->SetCustomAttenuationRadius(AttenuationRadius);
+				// LightNode->SetCustomAttenuationRadius(AttenuationRadius);
 
 				NodeContainer.AddNode(LightNode);
 			}
@@ -116,6 +488,7 @@ namespace UE::InterchangeUsdTranslator::Private
 		{
 			UInterchangeRectLightNode* LightNode = NewObject<UInterchangeRectLightNode>(&NodeContainer);
 			LightNode->InitializeNode(NodeUid, NodeName, EInterchangeNodeContainerType::TranslatedAsset);
+			LightNode->SetAssetName(NodeName);
 			NodeContainer.AddNode(LightNode);
 		}
 		// #ueent_todo:
@@ -124,76 +497,205 @@ namespace UE::InterchangeUsdTranslator::Private
 
 	void AddCameraNode(const UE::FUsdPrim& Prim, UInterchangeBaseNodeContainer& NodeContainer)
 	{
-		//FUsdCamera Camera;
-		//if (!UsdToUnreal::ConvertGeomCamera(Prim, Camera, UsdUtils::GetEarliestTimeCode()))
+		// FUsdCamera Camera;
+		// if (!UsdToUnreal::ConvertGeomCamera(Prim, Camera, UsdUtils::GetEarliestTimeCode()))
 		//{
 		//	return;
-		//}
+		// }
 
-		FString NodeUid = TEXT("Camera:") + Prim.GetPrimPath().GetString();
+		FString NodeUid = CameraPrefix + Prim.GetPrimPath().GetString();
 		FString NodeName(Prim.GetName().ToString());
 
 		UInterchangePhysicalCameraNode* CameraNode = NewObject<UInterchangePhysicalCameraNode>(&NodeContainer);
 		CameraNode->InitializeNode(NodeUid, NodeName, EInterchangeNodeContainerType::TranslatedAsset);
 		NodeContainer.AddNode(CameraNode);
 
-		//CameraNode->SetCustomFocalLength(Camera.FocalLength);
-		//CameraNode->SetCustomSensorHeight(Camera.SensorHeight);
-		//CameraNode->SetCustomSensorWidth(Camera.SensorWidth);
+		// CameraNode->SetCustomFocalLength(Camera.FocalLength);
+		// CameraNode->SetCustomSensorHeight(Camera.SensorHeight);
+		// CameraNode->SetCustomSensorWidth(Camera.SensorWidth);
 	}
 
-	void AddSceneNode(const UE::FUsdStage& UsdStage, const UE::FUsdPrim& Prim, UInterchangeBaseNodeContainer& NodeContainer, const UInterchangeBaseNode* ParentNode)
+	void AddMeshNode(const UE::FUsdPrim& Prim, UInterchangeUSDTranslatorImpl* TranslatorImpl, UInterchangeBaseNodeContainer& NodeContainer)
 	{
-#if USE_USD_SDK
-		FTransform Transform = FTransform::Identity;
-		if (UsdToUnreal::ConvertXformable(UsdStage, UE::FUsdTyped(Prim), Transform, UsdUtils::GetEarliestTimeCode(), nullptr))
+		FString PrimPath = Prim.GetPrimPath().GetString();
+		FString NodeUid = MeshPrefix + PrimPath;
+		FString NodeName(Prim.GetName().ToString());
+
+		// Check if Node already exist with this ID
+		if (const UInterchangeMeshNode* Node = Cast<UInterchangeMeshNode>(NodeContainer.GetNode(NodeUid)))
 		{
-			FString NodeUid = Prim.GetPrimPath().GetString();
-			FString DisplayLabel(Prim.GetName().ToString());
+			return;
+		}
 
-			UInterchangeSceneNode* SceneNode = NewObject<UInterchangeSceneNode>(&NodeContainer);
-			SceneNode->InitializeNode(NodeUid, DisplayLabel, EInterchangeNodeContainerType::TranslatedScene);
-			SceneNode->SetAssetName(DisplayLabel);
-			SceneNode->SetCustomLocalTransform(&NodeContainer, Transform);
-			NodeContainer.AddNode(SceneNode);
+		UInterchangeMeshNode* MeshNode = NewObject<UInterchangeMeshNode>(&NodeContainer);
+		MeshNode->InitializeNode(NodeUid, NodeName, EInterchangeNodeContainerType::TranslatedAsset);
+		MeshNode->SetPayLoadKey(PrimPath, EInterchangeMeshPayLoadType::STATIC);
+		MeshNode->SetAssetName(NodeName);
 
-			FString Prefix;
-			if (Prim.IsA(TEXT("Camera")))
-			{
-				Prefix = TEXT("Camera:");
-				AddCameraNode(Prim, NodeContainer);
-			}
-			else if (Prim.HasAPI(TEXT("LightAPI")))
+		// Material assignments
+		{
+			const double TimeCode = UsdUtils::GetDefaultTimeCode();
+			const bool bProvideMaterialIndices = false;
+			UsdUtils::FUsdPrimMaterialAssignmentInfo Assignments = UsdUtils::GetPrimMaterialAssignments(
+				Prim,
+				TimeCode,
+				bProvideMaterialIndices,
+				TranslatorImpl->CachedMeshConversionOptions.RenderContext,
+				TranslatorImpl->CachedMeshConversionOptions.MaterialPurpose
+			);
+
+			for (const UsdUtils::FUsdPrimMaterialSlot& Slot : Assignments.Slots)
 			{
-				Prefix = TEXT("Light:");
-				AddLightNode(Prim, NodeContainer);
+				const FString& SlotName = Slot.SlotName;
+
+				// Get the Uid of the material instance that we'll end up assigning to this slot
+				FString MaterialInstanceUid;
+				switch (Slot.AssignmentType)
+				{
+					case UsdUtils::EPrimAssignmentType::DisplayColor:
+					{
+						AddDisplayColorMaterialInstanceNodeIfNeeded(NodeContainer, Slot.MaterialSource);
+						MaterialInstanceUid = MaterialPrefix + Slot.MaterialSource;	   // This is e.g. "!DisplayColor_0_1"
+						break;
+					}
+					case UsdUtils::EPrimAssignmentType::MaterialPrim:
+					{
+						MaterialInstanceUid = MaterialPrefix + Slot.MaterialSource;	   // This is the prim path
+						break;
+					}
+					case UsdUtils::EPrimAssignmentType::UnrealMaterial:
+					{
+						// TODO: We can't support these yet without a custom pipeline unfortunately
+						// We could spawn a material instance of the referenced material... That's probably not what you'd expect though
+						break;
+					}
+					default:
+					{
+						ensure(false);
+						break;
+					}
+				}
+
+				MeshNode->SetSlotMaterialDependencyUid(SlotName, MaterialInstanceUid);
 			}
+		}
+
+		NodeContainer.AddNode(MeshNode);
+	}
+
+	void Traverse(
+		const UE::FUsdStage& UsdStage,
+		const UE::FUsdPrim& Prim,
+		UInterchangeUSDTranslatorImpl* TranslatorImpl,
+		UInterchangeBaseNodeContainer& NodeContainer,
+		UInterchangeBaseNode* ParentNode
+	)
+	{
+		// Ignore prim subtrees from disabled purposes
+		// TODO: Move this to the pipeline and filter only the factory nodes
+		EUsdPurpose PrimPurpose = IUsdPrim::GetPurpose(Prim);
+		if (!EnumHasAllFlags(TranslatorImpl->CachedMeshConversionOptions.PurposesToLoad, PrimPurpose))
+		{
+			return;
+		}
+
+		FString NodeUid = Prim.GetPrimPath().GetString();
+		FString DisplayLabel(Prim.GetName().ToString());
 
-			if (!Prefix.IsEmpty())
+		// Generate asset node if applicable
+		const FString* Prefix = nullptr;
+		if (Prim.IsA(TEXT("Material")))
+		{
+			Prefix = &MaterialPrefix;
+			AddMaterialInstanceNode(Prim, TranslatorImpl, NodeContainer);
+		}
+		else if (Prim.IsA(TEXT("Mesh")))
+		{
+			Prefix = &MeshPrefix;
+			AddMeshNode(Prim, TranslatorImpl, NodeContainer);
+		}
+		else if (Prim.IsA(TEXT("Camera")))
+		{
+			Prefix = &CameraPrefix;
+			AddCameraNode(Prim, NodeContainer);
+		}
+		else if (Prim.HasAPI(TEXT("LightAPI")))
+		{
+			Prefix = &LightPrefix;
+			AddLightNode(Prim, NodeContainer);
+		}
+
+		// Generate scene node if we're an Xformable
+		UInterchangeSceneNode* SceneNode = nullptr;
+		FTransform Transform = FTransform::Identity;
+		if (UsdToUnreal::ConvertXformable(UsdStage, UE::FUsdTyped(Prim), Transform, UsdUtils::GetDefaultTimeCode(), nullptr))
+		{
+			SceneNode = NewObject<UInterchangeSceneNode>(&NodeContainer);
+			SceneNode->InitializeNode(NodeUid, DisplayLabel, EInterchangeNodeContainerType::TranslatedScene);
+			SceneNode->SetCustomLocalTransform(&NodeContainer, Transform);
+			NodeContainer.AddNode(SceneNode);
+
+			// Connect scene node and asset node
+			if (Prefix)
 			{
-				const FString AssetNodeUid = Prefix + NodeUid;
+				const FString AssetNodeUid = *Prefix + NodeUid;
 				SceneNode->SetCustomAssetInstanceUid(AssetNodeUid);
 			}
 
+			// Connect parent and child scene nodes
 			if (ParentNode)
 			{
 				NodeContainer.SetNodeParentUid(SceneNode->GetUniqueID(), ParentNode->GetUniqueID());
 			}
-
-			ParentNode = SceneNode;
 		}
 
+		// Recurse into child prims
 		for (const FUsdPrim& ChildPrim : Prim.GetChildren())
 		{
-			AddSceneNode(UsdStage, ChildPrim, NodeContainer, ParentNode);
+			Traverse(UsdStage, ChildPrim, TranslatorImpl, NodeContainer, SceneNode);
 		}
-#endif // USE_USD_SDK
 	}
-}
+
+	bool GetStaticMeshPayloadDataForPayLoadKey(
+		const UE::FUsdStage& UsdStage,
+		const FString& PrimPath,
+		const UsdToUnreal::FUsdMeshConversionOptions& Options,
+		FMeshDescription& OutMeshDescription
+	)
+	{
+		UE::FUsdPrim Prim = UsdStage.GetPrimAtPath(UE::FSdfPath{*PrimPath});
+		if (!Prim)
+		{
+			return false;
+		}
+
+		// TODO: We can't do much with these yet: There will be used to generate primvar-compatible
+		// versions of the materials that are assigned to this mesh, whenever we get a pipeline
+		UsdUtils::FUsdPrimMaterialAssignmentInfo TempMaterialInfo;
+		return UsdToUnreal::ConvertGeomMesh(Prim, OutMeshDescription, TempMaterialInfo, Options);
+	}
+
+#endif
+}	 // namespace UE::InterchangeUsdTranslator::Private
 
 bool UInterchangeUSDTranslator::Translate(UInterchangeBaseNodeContainer& NodeContainer) const
 {
+#if USE_USD_SDK
 	using namespace UE;
+	using namespace UE::InterchangeUsdTranslator::Private;
+	using namespace UsdToUnreal;
+
+	UInterchangeUSDTranslatorImpl* ImplPtr = Impl.Get();
+	if (!ImplPtr)
+	{
+		return false;
+	}
+
+	UInterchangeUsdTranslatorSettings* Settings = Cast<UInterchangeUsdTranslatorSettings>(GetSettings());
+	if (!Settings)
+	{
+		return false;
+	}
 
 	FString FilePath = GetSourceData()->GetFilename();
 	if (!FPaths::FileExists(FilePath))
@@ -201,33 +703,302 @@ bool UInterchangeUSDTranslator::Translate(UInterchangeBaseNodeContainer& NodeCon
 		return false;
 	}
 
-	FUsdStage UsdStage = UnrealUSDWrapper::OpenStage(*FilePath, EUsdInitialLoadSet::LoadAll);
-	if (!UsdStage)
+	// Import should always feel like it's directly from disk, so we ignore already loaded layers and stage cache
+	const bool bUseStageCache = false;
+	const bool bForceReloadLayersFromDisk = true;
+	ImplPtr->UsdStage = UnrealUSDWrapper::OpenStage(*FilePath, EUsdInitialLoadSet::LoadAll, bUseStageCache, bForceReloadLayersFromDisk);
+	if (!ImplPtr->UsdStage)
 	{
 		return false;
 	}
 
-	// Scene hierarchy
-	FUsdPrim RootPrim = UsdStage.GetPseudoRoot();
+	// Apply coordinate system conversion to the stage if we have one
+	if (Settings && Settings->bOverrideStageOptions)
+	{
+		UsdUtils::SetUsdStageMetersPerUnit(ImplPtr->UsdStage, Settings->StageOptions.MetersPerUnit);
+		UsdUtils::SetUsdStageUpAxis(ImplPtr->UsdStage, Settings->StageOptions.UpAxis);
+	}
+
+	// Cache these so we don't have to keep converting these tokens over and over during translation
+	FUsdMeshConversionOptions& MeshOptions = ImplPtr->CachedMeshConversionOptions;
+	MeshOptions.PurposesToLoad = (EUsdPurpose)Settings->GeometryPurpose;
+	MeshOptions.RenderContext = Settings->RenderContext.IsNone() ? pxr::UsdShadeTokens->universalRenderContext
+																 : UnrealToUsd::ConvertToken(*Settings->RenderContext.ToString()).Get();
+	MeshOptions.MaterialPurpose = Settings->MaterialPurpose.IsNone() ? pxr::UsdShadeTokens->allPurpose
+																	 : UnrealToUsd::ConvertToken(*Settings->MaterialPurpose.ToString()).Get();
+
+	// Traverse stage and emit translated nodes
+	FUsdPrim RootPrim = ImplPtr->UsdStage.GetPseudoRoot();
 	for (const FUsdPrim& Prim : RootPrim.GetChildren())
 	{
-		InterchangeUsdTranslator::Private::AddSceneNode(UsdStage, Prim, NodeContainer, nullptr);
+		UInterchangeBaseNode* ParentNode = nullptr;
+		Traverse(ImplPtr->UsdStage, Prim, ImplPtr, NodeContainer, ParentNode);
 	}
 
 	return true;
+#else
+	return false;
+#endif	  // USE_USD_SDK
 }
 
-TFuture<TOptional<UE::Interchange::FMeshPayloadData>> UInterchangeUSDTranslator::GetMeshPayloadData(const FInterchangeMeshPayLoadKey& PayLoadKey, const FTransform& MeshGlobalTransform) const
+void UInterchangeUSDTranslator::ReleaseSource()
 {
-	return {};
+	UE::InterchangeUsdTranslator::Private::UInterchangeUSDTranslatorImpl* ImplPtr = Impl.Get();
+	if (!ImplPtr)
+	{
+		return;
+	}
+
+	ImplPtr->UsdStage = UE::FUsdStage{};
+
+	if (ImplPtr->TranslatorSettings)
+	{
+		ImplPtr->TranslatorSettings->ClearFlags(RF_Standalone);
+		ImplPtr->TranslatorSettings = nullptr;
+	}
 }
 
-TOptional<UE::Interchange::FImportImage> UInterchangeUSDTranslator::GetTexturePayloadData(const FString& PayloadKey, TOptional<FString>& AlternateTexturePath) const
+UInterchangeTranslatorSettings* UInterchangeUSDTranslator::GetSettings() const
 {
-	return {};
+	using namespace UE::InterchangeUsdTranslator::Private;
+
+	UInterchangeUSDTranslatorImpl* ImplPtr = Impl.Get();
+	if (!ImplPtr)
+	{
+		return nullptr;
+	}
+
+	TObjectPtr<UInterchangeUsdTranslatorSettings>& Settings = ImplPtr->TranslatorSettings;
+	if (!Settings)
+	{
+		Settings = DuplicateObject<UInterchangeUsdTranslatorSettings>(
+			UInterchangeUsdTranslatorSettings::StaticClass()->GetDefaultObject<UInterchangeUsdTranslatorSettings>(),
+			GetTransientPackage()
+		);
+		Settings->LoadSettings();
+		Settings->SetFlags(RF_Standalone);
+		Settings->ClearInternalFlags(EInternalObjectFlags::Async);
+	}
+	return Settings;
+}
+
+void UInterchangeUSDTranslator::SetSettings(const UInterchangeTranslatorSettings* InterchangeTranslatorSettings)
+{
+	using namespace UE::InterchangeUsdTranslator::Private;
+
+	UInterchangeUSDTranslatorImpl* ImplPtr = Impl.Get();
+	if (!ImplPtr)
+	{
+		return;
+	}
+
+	TObjectPtr<UInterchangeUsdTranslatorSettings>& Settings = ImplPtr->TranslatorSettings;
+
+	if (Settings)
+	{
+		Settings->ClearFlags(RF_Standalone);
+		Settings->ClearInternalFlags(EInternalObjectFlags::Async);
+		Settings = nullptr;
+	}
+	if (const UInterchangeUsdTranslatorSettings* USDTranslatorSettings = Cast<UInterchangeUsdTranslatorSettings>(InterchangeTranslatorSettings))
+	{
+		Settings = DuplicateObject<UInterchangeUsdTranslatorSettings>(USDTranslatorSettings, GetTransientPackage());
+		Settings->ClearInternalFlags(EInternalObjectFlags::Async);
+		Settings->SetFlags(RF_Standalone);
+	}
+}
+
+TFuture<TOptional<UE::Interchange::FMeshPayloadData>> UInterchangeUSDTranslator::GetMeshPayloadData(
+	const FInterchangeMeshPayLoadKey& PayloadKey,
+	const FTransform& MeshGlobalTransform
+) const
+{
+	using namespace UE::InterchangeUsdTranslator::Private;
+
+	return Async(
+		EAsyncExecution::TaskGraph,
+		[this, PayloadKey, MeshGlobalTransform]
+		{
+			TOptional<UE::Interchange::FMeshPayloadData> Result;
+
+#if USE_USD_SDK
+			UInterchangeUSDTranslatorImpl* ImplPtr = Impl.Get();
+			if (!ImplPtr)
+			{
+				return Result;
+			}
+
+			UsdToUnreal::FUsdMeshConversionOptions OptionsCopy = ImplPtr->CachedMeshConversionOptions;
+			OptionsCopy.AdditionalTransform = MeshGlobalTransform;
+
+			bool bSuccess = false;
+			UE::Interchange::FMeshPayloadData MeshPayloadData;
+			switch (PayloadKey.Type)
+			{
+				case EInterchangeMeshPayLoadType::STATIC:
+				{
+					bSuccess = UE::InterchangeUsdTranslator::Private::GetStaticMeshPayloadDataForPayLoadKey(
+						ImplPtr->UsdStage,
+						PayloadKey.UniqueId,
+						OptionsCopy,
+						MeshPayloadData.MeshDescription
+					);
+					break;
+				}
+				case EInterchangeMeshPayLoadType::SKELETAL:		  // Fallthrough
+				case EInterchangeMeshPayLoadType::MORPHTARGET:	  // Fallthrough
+				case EInterchangeMeshPayLoadType::NONE:			  // Fallthrough
+				default:
+					break;
+			}
+
+			if (bSuccess)
+			{
+				Result.Emplace(MeshPayloadData);
+			}
+#endif	  // USE_USD_SDK
+
+			return Result;
+		}
+	);
+}
+
+TOptional<UE::Interchange::FImportImage> UInterchangeUSDTranslator::GetTexturePayloadData(
+	const FString& PayloadKey,
+	TOptional<FString>& AlternateTexturePath
+) const
+{
+	using namespace UE::InterchangeUsdTranslator::Private;
+
+	TOptional<UE::Interchange::FImportImage> TexturePayloadData;
+
+#if USE_USD_SDK
+	FString FilePath;
+	TextureGroup TextureGroup;
+	bool bDecoded = DecodeTexturePayloadKey(PayloadKey, FilePath, TextureGroup);
+	if (!bDecoded)
+	{
+		return {};
+	}
+
+	// Defer back to another translator to actually parse the texture raw data
+	UE::Interchange::Private::FScopedTranslator ScopedTranslator(FilePath, Results);
+	const IInterchangeTexturePayloadInterface* TextureTranslator = ScopedTranslator.GetPayLoadInterface<IInterchangeTexturePayloadInterface>();
+	if (!ensure(TextureTranslator))
+	{
+		return {};
+	}
+
+	AlternateTexturePath = FilePath;
+
+	// The texture translators don't use the payload key, and read the texture directly from the SourceData's file path
+	const FString UnusedPayloadKey = {};
+	TexturePayloadData = TextureTranslator->GetTexturePayloadData(UnusedPayloadKey, AlternateTexturePath);
+
+	// Move compression settings onto the payload data.
+	// Note: We don't author anything else on the texture payload data here (like the sRGB flag), because those
+	// settings were already on our translated node, and presumably already made their way to the factory node.
+	// The factory should use them to override whatever it finds in this payload data, with the exception of the
+	// compression settings (which can't be stored on the translated node)
+	TexturePayloadData->CompressionSettings = TextureGroup == TEXTUREGROUP_WorldNormalMap ? TC_Normalmap : TC_Default;
+#endif	  // USE_USD_SDK
+
+	return TexturePayloadData;
+}
+
+TOptional<UE::Interchange::FImportBlockedImage> UInterchangeUSDTranslator::GetBlockedTexturePayloadData(
+	const FString& PayloadKey,
+	TOptional<FString>& AlternateTexturePath
+) const
+{
+	using namespace UE::InterchangeUsdTranslator::Private;
+
+	UE::Interchange::FImportBlockedImage BlockData;
+
+#if USE_USD_SDK
+	FString FilePath;
+	TextureGroup TextureGroup;
+	bool bDecoded = DecodeTexturePayloadKey(PayloadKey, FilePath, TextureGroup);
+	if (!bDecoded)
+	{
+		return {};
+	}
+
+	AlternateTexturePath = FilePath;
+
+	// Collect all the UDIM tile filepaths similar to this current tile. If we've been asked to translate
+	// a blocked texture then we must have some
+	TMap<int32, FString> TileIndexToPath = UE::TextureUtilitiesCommon::GetUDIMBlocksFromSourceFile(
+		FilePath,
+		UE::TextureUtilitiesCommon::DefaultUdimRegexPattern
+	);
+	if (!ensure(TileIndexToPath.Num() > 0))
+	{
+		return {};
+	}
+
+	bool bInitializedBlockData = false;
+
+	TArray<UE::Interchange::FImportImage> TileImages;
+	TileImages.Reserve(TileIndexToPath.Num());
+
+	for (const TPair<int32, FString>& TileIndexAndPath : TileIndexToPath)
+	{
+		int32 UdimTile = TileIndexAndPath.Key;
+		const FString& TileFilePath = TileIndexAndPath.Value;
+
+		int32 BlockX = INDEX_NONE;
+		int32 BlockY = INDEX_NONE;
+		UE::TextureUtilitiesCommon::ExtractUDIMCoordinates(UdimTile, BlockX, BlockY);
+		if (BlockX == INDEX_NONE || BlockY == INDEX_NONE)
+		{
+			continue;
+		}
+
+		// Find another translator that actually supports that filetype to handle the texture
+		UE::Interchange::Private::FScopedTranslator ScopedTranslator(TileFilePath, Results);
+		const IInterchangeTexturePayloadInterface* TextureTranslator = ScopedTranslator.GetPayLoadInterface<IInterchangeTexturePayloadInterface>();
+		if (!ensure(TextureTranslator))
+		{
+			continue;
+		}
+
+		// Invoke the translator to actually load the texture and parse it
+		const FString UnusedPayloadKey = {};
+		TOptional<UE::Interchange::FImportImage> TexturePayloadData;
+		TexturePayloadData = TextureTranslator->GetTexturePayloadData(UnusedPayloadKey, AlternateTexturePath);
+		if (!TexturePayloadData)
+		{
+			continue;
+		}
+		const UE::Interchange::FImportImage& Image = TileImages.Emplace_GetRef(MoveTemp(TexturePayloadData.GetValue()));
+		TexturePayloadData.Reset();
+
+		// Initialize the settings on the BlockData itself based on the first image we parse
+		if (!bInitializedBlockData)
+		{
+			bInitializedBlockData = true;
+
+			BlockData.Format = Image.Format;
+			BlockData.CompressionSettings = TextureGroup == TEXTUREGROUP_WorldNormalMap ? TC_Normalmap : TC_Default;
+			BlockData.bSRGB = Image.bSRGB;
+			BlockData.MipGenSettings = Image.MipGenSettings;
+		}
+
+		// Prepare the BlockData to receive this image data (later)
+		BlockData.InitBlockFromImage(BlockX, BlockY, Image);
+	}
+
+	// Move all of the FImportImage buffers into the BlockData itself
+	BlockData.MigrateDataFromImagesToRawData(TileImages);
+#endif	  // USE_USD_SDK
+
+	return BlockData;
 }
 
-TArray<UE::Interchange::FAnimationPayloadData> UInterchangeUSDTranslator::GetAnimationPayloadData(const TArray<UE::Interchange::FAnimationPayloadQuery>& PayloadQuery) const
+TArray<UE::Interchange::FAnimationPayloadData> UInterchangeUSDTranslator::GetAnimationPayloadData(
+	const TArray<UE::Interchange::FAnimationPayloadQuery>& PayloadQuery
+) const
 {
 	return {};
 }
diff --git a/Engine/Plugins/Interchange/Runtime/Source/Import/Public/Usd/InterchangeUsdTranslator.h b/Engine/Plugins/Interchange/Runtime/Source/Import/Public/Usd/InterchangeUsdTranslator.h
--- a/Engine/Plugins/Interchange/Runtime/Source/Import/Public/Usd/InterchangeUsdTranslator.h
+++ b/Engine/Plugins/Interchange/Runtime/Source/Import/Public/Usd/InterchangeUsdTranslator.h
@@ -2,20 +2,74 @@
 
 #pragma once
 
-#include "InterchangeTranslatorBase.h"
 #include "Animation/InterchangeAnimationPayloadInterface.h"
+#include "InterchangeTranslatorBase.h"
 #include "Mesh/InterchangeMeshPayloadInterface.h"
+#include "Texture/InterchangeBlockedTexturePayloadInterface.h"
 #include "Texture/InterchangeTexturePayloadInterface.h"
 #include "UObject/Object.h"
 #include "UObject/ObjectMacros.h"
 
+#include "USDStageOptions.h"
+
 #include "InterchangeUsdTranslator.generated.h"
 
-/* For now, USD Interchange (FBX parity) translator supports textures/materials, static/skeletal meshes, and animation*/
+namespace UE::InterchangeUsdTranslator::Private
+{
+	class UInterchangeUSDTranslatorImpl;
+}
+
+UCLASS(BlueprintType, editinlinenew, MinimalAPI)
+class UInterchangeUsdTranslatorSettings : public UInterchangeTranslatorSettings
+{
+	GENERATED_BODY()
+
+public:
+
+	/**
+	 * TODO: Most (if not all) of these settings should probably be specific to the pipeline, and not the
+	 * translator, and will be moved there whenever we do make a USD pipeline.
+	 *
+	 * For example: Instead of filtering prims based on purpose on translation, we should emit all of the translated
+	 * nodes and filter only later on the pipeline, as translating the scene should be fast either way. That way users
+	 * can even customize/disable that behavior if they want to, in order to use their own pipelines.
+	 *
+	 * It's not clear what to do about StageOptions or RenderContext though: Maybe these should be here, as they
+	 * actively affect how we translate the scene? (e.g. the generated Interchange Material node for a Material
+	 * prim will be very different whether we use "universal", "unreal" or "mtlx" render contexts...)
+	 */
+
+	/** Only import geometry prims with these specific purposes from the USD file */
+	UPROPERTY(EditAnywhere, Category = "USD Translator", meta = (Bitmask, BitmaskEnum = "/Script/UnrealUSDWrapper.EUsdPurpose"))
+	int32 GeometryPurpose;
+
+	/** Specifies which set of shaders to use when parsing USD materials, in addition to the universal render context. */
+	UPROPERTY(EditAnywhere, Category = "USD Translator")
+	FName RenderContext;
+
+	/** Specifies which material purpose to use when parsing USD material bindings, in addition to the "allPurpose" fallback */
+	UPROPERTY(EditAnywhere, Category = "USD Translator")
+	FName MaterialPurpose;
+
+	/** Whether to use the specified StageOptions instead of the stage's own settings */
+	UPROPERTY(EditAnywhere, Category = "USD Translator")
+	bool bOverrideStageOptions;
+
+	/** Custom StageOptions to use for the stage */
+	UPROPERTY(EditAnywhere, Category = "USD Translator", meta = (EditCondition = bOverrideStageOptions))
+	FUsdStageOptions StageOptions;
+
+public:
+	UInterchangeUsdTranslatorSettings();
+};
+
+/* For now, USD Interchange (FBX parity) translator supports textures, materials and static meshes */
 UCLASS(BlueprintType)
-class UInterchangeUSDTranslator : public UInterchangeTranslatorBase
+class UInterchangeUSDTranslator
+	: public UInterchangeTranslatorBase
 	, public IInterchangeMeshPayloadInterface
 	, public IInterchangeTexturePayloadInterface
+	, public IInterchangeBlockedTexturePayloadInterface
 	, public IInterchangeAnimationPayloadInterface
 {
 	GENERATED_BODY()
@@ -27,10 +81,30 @@ public:
 	virtual EInterchangeTranslatorType GetTranslatorType() const override;
 	virtual EInterchangeTranslatorAssetType GetSupportedAssetTypes() const override;
 	virtual TArray<FString> GetSupportedFormats() const override;
-	virtual bool Translate( UInterchangeBaseNodeContainer& BaseNodeContainer ) const override;
+	virtual bool Translate(UInterchangeBaseNodeContainer& BaseNodeContainer) const override;
+	virtual void ReleaseSource() override;
+	virtual UInterchangeTranslatorSettings* GetSettings() const override;
+	virtual void SetSettings(const UInterchangeTranslatorSettings* InterchangeTranslatorSettings) override;
 	/** End UInterchangeTranslatorBase API*/
 
-	virtual TFuture<TOptional<UE::Interchange::FMeshPayloadData>> GetMeshPayloadData(const FInterchangeMeshPayLoadKey& PayLoadKey, const FTransform& MeshGlobalTransform) const override;
-	virtual TOptional<UE::Interchange::FImportImage> GetTexturePayloadData(const FString& PayloadKey, TOptional<FString>& AlternateTexturePath) const override;
-	virtual TArray<UE::Interchange::FAnimationPayloadData> GetAnimationPayloadData(const TArray<UE::Interchange::FAnimationPayloadQuery>& PayloadQuery) const override;
+	/** Begin Interchange payload interfaces */
+	virtual TFuture<TOptional<UE::Interchange::FMeshPayloadData>> GetMeshPayloadData(
+		const FInterchangeMeshPayLoadKey& PayLoadKey,
+		const FTransform& MeshGlobalTransform
+	) const override;
+
+	virtual TOptional<UE::Interchange::FImportImage> GetTexturePayloadData(const FString& PayloadKey, TOptional<FString>& AlternateTexturePath)
+		const override;
+
+	virtual TOptional<UE::Interchange::FImportBlockedImage> GetBlockedTexturePayloadData(
+		const FString& PayloadKey,
+		TOptional<FString>& AlternateTexturePath
+	) const override;
+
+	virtual TArray<UE::Interchange::FAnimationPayloadData> GetAnimationPayloadData(const TArray<UE::Interchange::FAnimationPayloadQuery>& PayloadQuery
+	) const override;
+	/** End Interchange payload interfaces */
+
+private:
+	TUniquePtr<UE::InterchangeUsdTranslator::Private::UInterchangeUSDTranslatorImpl> Impl;
 };
diff --git a/Engine/Plugins/Runtime/USDCore/Source/USDClasses/Private/USDClassesModule.cpp b/Engine/Plugins/Runtime/USDCore/Source/USDClasses/Private/USDClassesModule.cpp
--- a/Engine/Plugins/Runtime/USDCore/Source/USDClasses/Private/USDClassesModule.cpp
+++ b/Engine/Plugins/Runtime/USDCore/Source/USDClasses/Private/USDClassesModule.cpp
@@ -18,7 +18,6 @@
 #include "Engine/SkinnedAssetCommon.h"
 #include "Engine/StaticMesh.h"
 #include "Engine/Texture.h"
-#include "Engine/Texture2D.h"
 #include "EngineAnalytics.h"
 #include "Framework/Notifications/NotificationManager.h"
 #include "GeometryCache.h"
@@ -28,8 +27,6 @@
 #include "HAL/FileManager.h"
 #include "Materials/Material.h"
 #include "Materials/MaterialInstance.h"
-#include "Materials/MaterialInstanceConstant.h"
-#include "Materials/MaterialInstanceDynamic.h"
 #include "Materials/MaterialInterface.h"
 #include "Misc/EngineVersion.h"
 #include "Misc/FileHelper.h"
@@ -39,7 +36,7 @@
 #include "Modules/ModuleManager.h"
 #include "PhysicsEngine/PhysicsAsset.h"
 #include "Serialization/JsonSerializer.h"
-#include "Sound/SoundWave.h"
+#include "Sound/SoundBase.h"
 #include "SparseVolumeTexture/SparseVolumeTexture.h"
 #include "UObject/NameTypes.h"
 #include "UObject/ObjectSaveContext.h"
@@ -58,8 +55,6 @@ DEFINE_LOG_CATEGORY(LogUsd);
 
 namespace UE::USDClasses::Private
 {
-	static const FString DisplayColorID = TEXT("!DisplayColor");
-
 	TSharedPtr<FJsonObject> ParseJSON(const FString& FileContents)
 	{
 		if (FileContents.IsEmpty())
@@ -599,117 +594,19 @@ FString IUsdClassesModule::SanitizeObjectName(const FString& InObjectName)
 	return UsdUnreal::ObjectUtils::SanitizeObjectName(InObjectName);
 }
 
-FString IUsdClassesModule::FDisplayColorMaterial::ToString()
-{
-	return FString::Printf(TEXT("%s_%d_%d"), *UE::USDClasses::Private::DisplayColorID, bHasOpacity, bIsDoubleSided);
-}
-
-TOptional<IUsdClassesModule::FDisplayColorMaterial> IUsdClassesModule::FDisplayColorMaterial::FromString(const FString& DisplayColorString)
-{
-	TArray<FString> Tokens;
-	DisplayColorString.ParseIntoArray(Tokens, TEXT("_"));
-
-	if (Tokens.Num() != 3 || Tokens[0] != UE::USDClasses::Private::DisplayColorID)
-	{
-		return {};
-	}
-
-	IUsdClassesModule::FDisplayColorMaterial Result;
-	Result.bHasOpacity = static_cast<bool>(FCString::Atoi(*Tokens[1]));
-	Result.bIsDoubleSided = static_cast<bool>(FCString::Atoi(*Tokens[2]));
-	return Result;
-}
-
 const FSoftObjectPath* IUsdClassesModule::GetReferenceMaterialPath(const FDisplayColorMaterial& DisplayColorDescription)
 {
-	const UUsdProjectSettings* Settings = GetDefault<UUsdProjectSettings>();
-	if (!Settings)
-	{
-		return nullptr;
-	}
-
-	if (DisplayColorDescription.bHasOpacity)
-	{
-		if (DisplayColorDescription.bIsDoubleSided)
-		{
-			return &Settings->ReferenceDisplayColorAndOpacityTwoSidedMaterial;
-		}
-		else
-		{
-			return &Settings->ReferenceDisplayColorAndOpacityMaterial;
-		}
-	}
-	else
-	{
-		if (DisplayColorDescription.bIsDoubleSided)
-		{
-			return &Settings->ReferenceDisplayColorTwoSidedMaterial;
-		}
-		else
-		{
-			return &Settings->ReferenceDisplayColorMaterial;
-		}
-	}
+	return UsdUnreal::MaterialUtils::GetReferenceMaterialPath(DisplayColorDescription);
 }
 
 UMaterialInstanceDynamic* IUsdClassesModule::CreateDisplayColorMaterialInstanceDynamic(const FDisplayColorMaterial& DisplayColorDescription)
 {
-	const FSoftObjectPath* ParentPathPtr = GetReferenceMaterialPath(DisplayColorDescription);
-	if (!ParentPathPtr)
-	{
-		return nullptr;
-	}
-
-	if (UMaterialInterface* ParentMaterial = Cast<UMaterialInterface>(ParentPathPtr->TryLoad()))
-	{
-		FName AssetName = MakeUniqueObjectName(
-			GetTransientPackage(),
-			UMaterialInstanceConstant::StaticClass(),
-			*FString::Printf(
-				TEXT("DisplayColor%s%s"),
-				DisplayColorDescription.bHasOpacity ? TEXT("_Translucent") : TEXT(""),
-				DisplayColorDescription.bIsDoubleSided ? TEXT("_TwoSided") : TEXT("")
-			)
-		);
-
-		if (UMaterialInstanceDynamic* NewMaterial = UMaterialInstanceDynamic::Create(ParentMaterial, GetTransientPackage(), AssetName))
-		{
-			return NewMaterial;
-		}
-	}
-
-	return nullptr;
+	return UsdUnreal::MaterialUtils::CreateDisplayColorMaterialInstanceDynamic(DisplayColorDescription);
 }
 
 UMaterialInstanceConstant* IUsdClassesModule::CreateDisplayColorMaterialInstanceConstant(const FDisplayColorMaterial& DisplayColorDescription)
 {
-#if WITH_EDITOR
-	const FSoftObjectPath* ParentPathPtr = GetReferenceMaterialPath(DisplayColorDescription);
-	if (!ParentPathPtr)
-	{
-		return nullptr;
-	}
-
-	if (UMaterialInterface* ParentMaterial = Cast<UMaterialInterface>(ParentPathPtr->TryLoad()))
-	{
-		FName AssetName = MakeUniqueObjectName(
-			GetTransientPackage(),
-			UMaterialInstanceConstant::StaticClass(),
-			*FString::Printf(
-				TEXT("DisplayColor%s%s"),
-				DisplayColorDescription.bHasOpacity ? TEXT("_Translucent") : TEXT(""),
-				DisplayColorDescription.bIsDoubleSided ? TEXT("_TwoSided") : TEXT("")
-			)
-		);
-
-		if (UMaterialInstanceConstant* MaterialInstance = NewObject<UMaterialInstanceConstant>(GetTransientPackage(), AssetName, RF_NoFlags))
-		{
-			UMaterialEditingLibrary::SetMaterialInstanceParent(MaterialInstance, ParentMaterial);
-			return MaterialInstance;
-		}
-	}
-#endif	  // WITH_EDITOR
-	return nullptr;
+	return UsdUnreal::MaterialUtils::CreateDisplayColorMaterialInstanceConstant(DisplayColorDescription);
 }
 
 class FUsdClassesModule : public IUsdClassesModule
diff --git a/Engine/Plugins/Runtime/USDCore/Source/USDClasses/Private/USDDrawModeComponent.cpp b/Engine/Plugins/Runtime/USDCore/Source/USDClasses/Private/USDDrawModeComponent.cpp
--- a/Engine/Plugins/Runtime/USDCore/Source/USDClasses/Private/USDDrawModeComponent.cpp
+++ b/Engine/Plugins/Runtime/USDCore/Source/USDClasses/Private/USDDrawModeComponent.cpp
@@ -3,7 +3,7 @@
 #include "USDDrawModeComponent.h"
 
 #include "USDAssetUserData.h"
-#include "USDClassesModule.h"
+#include "USDMaterialUtils.h"
 #include "USDProjectSettings.h"
 
 #include "Components/BaseDynamicMeshSceneProxy.h"
@@ -1223,11 +1223,11 @@ void UUsdDrawModeComponent::RefreshMaterialInstances()
 		// Create a new display color instance
 		if (!DisplayColorInstance)
 		{
-			IUsdClassesModule::FDisplayColorMaterial Desc;
+			UsdUnreal::MaterialUtils::FDisplayColorMaterial Desc;
 			Desc.bHasOpacity = false;
 			Desc.bIsDoubleSided = false;
 
-			DisplayColorInstance = IUsdClassesModule::CreateDisplayColorMaterialInstanceDynamic(Desc);
+			DisplayColorInstance = UsdUnreal::MaterialUtils::CreateDisplayColorMaterialInstanceDynamic(Desc);
 
 			if (DisplayColorInstance)
 			{
diff --git a/Engine/Plugins/Runtime/USDCore/Source/USDClasses/Private/USDMaterialUtils.cpp b/Engine/Plugins/Runtime/USDCore/Source/USDClasses/Private/USDMaterialUtils.cpp
new file mode 100644
--- /dev/null
+++ b/Engine/Plugins/Runtime/USDCore/Source/USDClasses/Private/USDMaterialUtils.cpp
@@ -0,0 +1,299 @@
+// Copyright Epic Games, Inc. All Rights Reserved.
+
+#include "USDMaterialUtils.h"
+
+#include "USDProjectSettings.h"
+
+#include "Materials/MaterialInstanceConstant.h"
+#include "Materials/MaterialInstanceDynamic.h"
+#include "UObject/Package.h"
+
+#if WITH_EDITOR
+#include "MaterialEditingLibrary.h"
+#endif	  // WITH_EDITOR
+
+namespace UE::USDMaterialUtils::Private
+{
+	static const FString DisplayColorID = TEXT("!DisplayColor");
+}
+
+FString UsdUnreal::MaterialUtils::FDisplayColorMaterial::ToString() const
+{
+	return FString::Printf(TEXT("%s_%d_%d"), *UE::USDMaterialUtils::Private::DisplayColorID, bHasOpacity, bIsDoubleSided);
+}
+
+FString UsdUnreal::MaterialUtils::FDisplayColorMaterial::ToPrettyString() const
+{
+	return FString::Printf(TEXT("DisplayColor%s%s"), bHasOpacity ? TEXT("_Translucent") : TEXT(""), bIsDoubleSided ? TEXT("_TwoSided") : TEXT(""));
+}
+
+TOptional<UsdUnreal::MaterialUtils::FDisplayColorMaterial> UsdUnreal::MaterialUtils::FDisplayColorMaterial::FromString(
+	const FString& DisplayColorString
+)
+{
+	TArray<FString> Tokens;
+	DisplayColorString.ParseIntoArray(Tokens, TEXT("_"));
+
+	if (Tokens.Num() != 3 || Tokens[0] != UE::USDMaterialUtils::Private::DisplayColorID)
+	{
+		return {};
+	}
+
+	FDisplayColorMaterial Result;
+	Result.bHasOpacity = static_cast<bool>(FCString::Atoi(*Tokens[1]));
+	Result.bIsDoubleSided = static_cast<bool>(FCString::Atoi(*Tokens[2]));
+	return Result;
+}
+
+const FSoftObjectPath* UsdUnreal::MaterialUtils::GetReferenceMaterialPath(const FDisplayColorMaterial& DisplayColorDescription)
+{
+	const UUsdProjectSettings* Settings = GetDefault<UUsdProjectSettings>();
+	if (!Settings)
+	{
+		return nullptr;
+	}
+
+	if (DisplayColorDescription.bHasOpacity)
+	{
+		if (DisplayColorDescription.bIsDoubleSided)
+		{
+			return &Settings->ReferenceDisplayColorAndOpacityTwoSidedMaterial;
+		}
+		else
+		{
+			return &Settings->ReferenceDisplayColorAndOpacityMaterial;
+		}
+	}
+	else
+	{
+		if (DisplayColorDescription.bIsDoubleSided)
+		{
+			return &Settings->ReferenceDisplayColorTwoSidedMaterial;
+		}
+		else
+		{
+			return &Settings->ReferenceDisplayColorMaterial;
+		}
+	}
+}
+
+UMaterialInstanceDynamic* UsdUnreal::MaterialUtils::CreateDisplayColorMaterialInstanceDynamic(const FDisplayColorMaterial& DisplayColorDescription)
+{
+	const FSoftObjectPath* ParentPathPtr = GetReferenceMaterialPath(DisplayColorDescription);
+	if (!ParentPathPtr)
+	{
+		return nullptr;
+	}
+
+	if (UMaterialInterface* ParentMaterial = Cast<UMaterialInterface>(ParentPathPtr->TryLoad()))
+	{
+		FName AssetName = MakeUniqueObjectName(
+			GetTransientPackage(),
+			UMaterialInstanceConstant::StaticClass(),
+			*DisplayColorDescription.ToPrettyString()
+		);
+
+		if (UMaterialInstanceDynamic* NewMaterial = UMaterialInstanceDynamic::Create(ParentMaterial, GetTransientPackage(), AssetName))
+		{
+			return NewMaterial;
+		}
+	}
+
+	return nullptr;
+}
+
+UMaterialInstanceConstant* UsdUnreal::MaterialUtils::CreateDisplayColorMaterialInstanceConstant(const FDisplayColorMaterial& DisplayColorDescription)
+{
+#if WITH_EDITOR
+	const FSoftObjectPath* ParentPathPtr = GetReferenceMaterialPath(DisplayColorDescription);
+	if (!ParentPathPtr)
+	{
+		return nullptr;
+	}
+
+	if (UMaterialInterface* ParentMaterial = Cast<UMaterialInterface>(ParentPathPtr->TryLoad()))
+	{
+		FName AssetName = MakeUniqueObjectName(
+			GetTransientPackage(),
+			UMaterialInstanceConstant::StaticClass(),
+			*DisplayColorDescription.ToPrettyString()
+		);
+
+		if (UMaterialInstanceConstant* MaterialInstance = NewObject<UMaterialInstanceConstant>(GetTransientPackage(), AssetName, RF_NoFlags))
+		{
+			UMaterialEditingLibrary::SetMaterialInstanceParent(MaterialInstance, ParentMaterial);
+			return MaterialInstance;
+		}
+	}
+#endif	  // WITH_EDITOR
+	return nullptr;
+}
+
+FSoftObjectPath UsdUnreal::MaterialUtils::GetReferencePreviewSurfaceMaterial(EUsdReferenceMaterialProperties ReferenceMaterialProperties)
+{
+	const UUsdProjectSettings* Settings = GetDefault<UUsdProjectSettings>();
+	if (!Settings)
+	{
+		return {};
+	}
+
+	const bool bIsTranslucent = EnumHasAnyFlags(ReferenceMaterialProperties, EUsdReferenceMaterialProperties::Translucent);
+	const bool bIsVT = EnumHasAnyFlags(ReferenceMaterialProperties, EUsdReferenceMaterialProperties::VT);
+	const bool bIsTwoSided = EnumHasAnyFlags(ReferenceMaterialProperties, EUsdReferenceMaterialProperties::TwoSided);
+
+	const FSoftObjectPath* TargetMaterialPath = nullptr;
+	if (bIsTranslucent)
+	{
+		if (bIsVT)
+		{
+			if (bIsTwoSided)
+			{
+				return Settings->ReferencePreviewSurfaceTranslucentTwoSidedVTMaterial;
+			}
+			else
+			{
+				return Settings->ReferencePreviewSurfaceTranslucentVTMaterial;
+			}
+		}
+		else
+		{
+			if (bIsTwoSided)
+			{
+				return Settings->ReferencePreviewSurfaceTranslucentTwoSidedMaterial;
+			}
+			else
+			{
+				return Settings->ReferencePreviewSurfaceTranslucentMaterial;
+			}
+		}
+	}
+	else
+	{
+		if (bIsVT)
+		{
+			if (bIsTwoSided)
+			{
+				return Settings->ReferencePreviewSurfaceTwoSidedVTMaterial;
+			}
+			else
+			{
+				return Settings->ReferencePreviewSurfaceVTMaterial;
+			}
+		}
+		else
+		{
+			if (bIsTwoSided)
+			{
+				return Settings->ReferencePreviewSurfaceTwoSidedMaterial;
+			}
+			else
+			{
+				return Settings->ReferencePreviewSurfaceMaterial;
+			}
+		}
+	}
+}
+
+FSoftObjectPath UsdUnreal::MaterialUtils::GetVTVersionOfReferencePreviewSurfaceMaterial(const FSoftObjectPath& ReferenceMaterial)
+{
+	if (!ReferenceMaterial.IsValid())
+	{
+		return {};
+	}
+
+	const UUsdProjectSettings* Settings = GetDefault<UUsdProjectSettings>();
+	if (!Settings)
+	{
+		return {};
+	}
+
+	if (ReferenceMaterial.ToString().Contains(TEXT("VT"), ESearchCase::CaseSensitive, ESearchDir::FromEnd))
+	{
+		return ReferenceMaterial;
+	}
+	else if (ReferenceMaterial == Settings->ReferencePreviewSurfaceMaterial)
+	{
+		return Cast<UMaterialInterface>(Settings->ReferencePreviewSurfaceVTMaterial.TryLoad());
+	}
+	else if (ReferenceMaterial == Settings->ReferencePreviewSurfaceTwoSidedMaterial)
+	{
+		return Cast<UMaterialInterface>(Settings->ReferencePreviewSurfaceTwoSidedVTMaterial.TryLoad());
+	}
+	else if (ReferenceMaterial == Settings->ReferencePreviewSurfaceTranslucentMaterial)
+	{
+		return Cast<UMaterialInterface>(Settings->ReferencePreviewSurfaceTranslucentVTMaterial.TryLoad());
+	}
+	else if (ReferenceMaterial == Settings->ReferencePreviewSurfaceTranslucentTwoSidedMaterial)
+	{
+		return Cast<UMaterialInterface>(Settings->ReferencePreviewSurfaceTranslucentTwoSidedVTMaterial.TryLoad());
+	}
+
+	// We should only ever call this function with a ReferenceMaterial that matches one of the above paths
+	ensure(false);
+	return {};
+}
+
+FSoftObjectPath UsdUnreal::MaterialUtils::GetTwoSidedVersionOfReferencePreviewSurfaceMaterial(const FSoftObjectPath& ReferenceMaterial)
+{
+	if (!ReferenceMaterial.IsValid())
+	{
+		return {};
+	}
+
+	const UUsdProjectSettings* Settings = GetDefault<UUsdProjectSettings>();
+	if (!Settings)
+	{
+		return {};
+	}
+
+	if (ReferenceMaterial.ToString().Contains(TEXT("TwoSided"), ESearchCase::CaseSensitive, ESearchDir::FromEnd))
+	{
+		return ReferenceMaterial;
+	}
+	else if (ReferenceMaterial == Settings->ReferencePreviewSurfaceMaterial)
+	{
+		return Settings->ReferencePreviewSurfaceTwoSidedMaterial;
+	}
+	else if (ReferenceMaterial == Settings->ReferencePreviewSurfaceTranslucentMaterial)
+	{
+		return Settings->ReferencePreviewSurfaceTranslucentTwoSidedMaterial;
+	}
+	else if (ReferenceMaterial == Settings->ReferencePreviewSurfaceVTMaterial)
+	{
+		return Settings->ReferencePreviewSurfaceTwoSidedVTMaterial;
+	}
+	else if (ReferenceMaterial == Settings->ReferencePreviewSurfaceTranslucentVTMaterial)
+	{
+		return Settings->ReferencePreviewSurfaceTranslucentTwoSidedVTMaterial;
+	}
+
+	// We should only ever call this function with a ReferenceMaterial that matches one of the above paths
+	ensure(false);
+	return {};
+}
+
+bool UsdUnreal::MaterialUtils::IsReferencePreviewSurfaceMaterial(const FSoftObjectPath& Material)
+{
+	if (!Material.IsValid())
+	{
+		return false;
+	}
+
+	const UUsdProjectSettings* Settings = GetDefault<UUsdProjectSettings>();
+	if (!Settings)
+	{
+		return false;
+	}
+
+	TSet<FSoftObjectPath> ReferenceMaterials = {
+		Settings->ReferencePreviewSurfaceMaterial,
+		Settings->ReferencePreviewSurfaceTranslucentMaterial,
+		Settings->ReferencePreviewSurfaceTwoSidedMaterial,
+		Settings->ReferencePreviewSurfaceTranslucentTwoSidedMaterial,
+		Settings->ReferencePreviewSurfaceVTMaterial,
+		Settings->ReferencePreviewSurfaceTranslucentVTMaterial,
+		Settings->ReferencePreviewSurfaceTwoSidedVTMaterial,
+		Settings->ReferencePreviewSurfaceTranslucentTwoSidedVTMaterial};
+
+	return ReferenceMaterials.Contains(Material);
+}
diff --git a/Engine/Plugins/Runtime/USDCore/Source/USDClasses/Public/USDClassesModule.h b/Engine/Plugins/Runtime/USDCore/Source/USDClasses/Public/USDClassesModule.h
--- a/Engine/Plugins/Runtime/USDCore/Source/USDClasses/Public/USDClassesModule.h
+++ b/Engine/Plugins/Runtime/USDCore/Source/USDClasses/Public/USDClassesModule.h
@@ -2,8 +2,9 @@
 
 #pragma once
 
-#include "CoreMinimal.h"
+#include "USDMaterialUtils.h"	 // For backwards compatibility for a few releases since we moved the EUsdReferenceMaterialProperties enum
 
+#include "CoreMinimal.h"
 #include "Modules/ModuleInterface.h"
 
 class UMaterialInstanceConstant;
@@ -11,15 +12,6 @@ class UMaterialInstanceDynamic;
 class UUsdAssetCache3;
 struct FAnalyticsEventAttribute;
 
-enum class EUsdReferenceMaterialProperties : uint8
-{
-	None = 0,
-	Translucent = 1,
-	VT = 2,
-	TwoSided = 4
-};
-ENUM_CLASS_FLAGS(EUsdReferenceMaterialProperties)
-
 class IUsdClassesModule : public IModuleInterface
 {
 public:
@@ -68,18 +60,15 @@ public:
 	UE_DEPRECATED(5.5, "This function has been moved to USDObjectUtils.h")
 	USDCLASSES_API static FString SanitizeObjectName(const FString& InObjectName);
 
-	/** Describes the type of vertex color/DisplayColor material that we would need in order to render a prim's displayColor data as intended */
-	struct USDCLASSES_API FDisplayColorMaterial
-	{
-		bool bHasOpacity = false;
-		bool bIsDoubleSided = false;
-
-		FString ToString();
-		static TOptional<FDisplayColorMaterial> FromString(const FString& DisplayColorString);
-	};
+	// Backwards compatibility for 5.5 as we moved this struct into MaterialUtils
+	using FDisplayColorMaterial = UsdUnreal::MaterialUtils::FDisplayColorMaterial;
 
+	UE_DEPRECATED(5.5, "This function has been moved to USDMaterialUtils.h")
 	USDCLASSES_API static const FSoftObjectPath* GetReferenceMaterialPath(const FDisplayColorMaterial& DisplayColorDescription);
 
+	UE_DEPRECATED(5.5, "This function has been moved to USDMaterialUtils.h")
 	USDCLASSES_API static UMaterialInstanceDynamic* CreateDisplayColorMaterialInstanceDynamic(const FDisplayColorMaterial& DisplayColorDescription);
+
+	UE_DEPRECATED(5.5, "This function has been moved to USDMaterialUtils.h")
 	USDCLASSES_API static UMaterialInstanceConstant* CreateDisplayColorMaterialInstanceConstant(const FDisplayColorMaterial& DisplayColorDescription);
 };
diff --git a/Engine/Plugins/Runtime/USDCore/Source/USDClasses/Public/USDMaterialUtils.h b/Engine/Plugins/Runtime/USDCore/Source/USDClasses/Public/USDMaterialUtils.h
new file mode 100644
--- /dev/null
+++ b/Engine/Plugins/Runtime/USDCore/Source/USDClasses/Public/USDMaterialUtils.h
@@ -0,0 +1,60 @@
+// Copyright Epic Games, Inc. All Rights Reserved.
+
+#pragma once
+
+#include "CoreMinimal.h"
+
+class UMaterialInstanceDynamic;
+class UMaterialInstanceConstant;
+
+enum class EUsdReferenceMaterialProperties : uint8
+{
+	None = 0,
+	Translucent = 1,
+	VT = 2,
+	TwoSided = 4
+};
+ENUM_CLASS_FLAGS(EUsdReferenceMaterialProperties)
+
+namespace UsdUnreal::MaterialUtils
+{
+	/** Describes the type of vertex color/DisplayColor material that we would need in order to render a prim's displayColor data as intended */
+	struct USDCLASSES_API FDisplayColorMaterial
+	{
+		bool bHasOpacity = false;
+		bool bIsDoubleSided = false;
+
+		FString ToString() const;
+		FString ToPrettyString() const;
+		static TOptional<FDisplayColorMaterial> FromString(const FString& DisplayColorString);
+	};
+
+	USDCLASSES_API const FSoftObjectPath* GetReferenceMaterialPath(const FDisplayColorMaterial& DisplayColorDescription);
+
+	USDCLASSES_API UMaterialInstanceDynamic* CreateDisplayColorMaterialInstanceDynamic(const FDisplayColorMaterial& DisplayColorDescription);
+	USDCLASSES_API UMaterialInstanceConstant* CreateDisplayColorMaterialInstanceConstant(const FDisplayColorMaterial& DisplayColorDescription);
+
+	// Returns one of the alternatives of the UsdPreviewSurface reference material depending on the material overrides
+	// provided, and nullptr otherwise
+	USDCLASSES_API FSoftObjectPath GetReferencePreviewSurfaceMaterial(EUsdReferenceMaterialProperties ReferenceMaterialProperties);
+
+	// Returns the VT version of the provided UsdPreviewSurface ReferenceMaterial. Returns the provided ReferenceMaterial back if
+	// it is already a VT-capable reference material, and returns nullptr if ReferenceMaterial isn't one of our reference material
+	// alternatives.
+	// Example: Receives UsdPreviewSurfaceTwoSided -> Returns UsdPreviewSurfaceTwoSidedVT
+	// Example: Receives UsdPreviewSurfaceTwoSidedVT -> Returns UsdPreviewSurfaceTwoSidedVT
+	// Example: Receives SomeOtherReferenceMaterial -> Returns nullptr
+	USDCLASSES_API FSoftObjectPath GetVTVersionOfReferencePreviewSurfaceMaterial(const FSoftObjectPath& ReferenceMaterial);
+
+	// Returns the two-sided version of the provided UsdPreviewSurface ReferenceMaterial. Returns the provided ReferenceMaterial
+	// back if it is already a two-sided-capable reference material, and returns nullptr if ReferenceMaterial isn't one of our reference
+	// material alternatives.
+	// Example: Receives UsdPreviewSurfaceTranslucent -> Returns UsdPreviewSurfaceTwoSidedTranslucent
+	// Example: Receives UsdPreviewSurfaceTwoSidedTranslucent -> Returns UsdPreviewSurfaceTwoSidedTranslucent
+	// Example: Receives SomeOtherReferenceMaterial -> Returns nullptr
+	USDCLASSES_API FSoftObjectPath GetTwoSidedVersionOfReferencePreviewSurfaceMaterial(const FSoftObjectPath& ReferenceMaterial);
+
+	// Returns whether the Material is one of the UsdPreviewSurface reference materials (which can be reassigned by the
+	// user on a per project basis)
+	USDCLASSES_API bool IsReferencePreviewSurfaceMaterial(const FSoftObjectPath& ReferenceMaterial);
+};	  // namespace UsdUnreal::MaterialUtils
diff --git a/Engine/Plugins/Runtime/USDCore/Source/USDUtilities/Private/USDGeomMeshConversion.cpp b/Engine/Plugins/Runtime/USDCore/Source/USDUtilities/Private/USDGeomMeshConversion.cpp
--- a/Engine/Plugins/Runtime/USDCore/Source/USDUtilities/Private/USDGeomMeshConversion.cpp
+++ b/Engine/Plugins/Runtime/USDCore/Source/USDUtilities/Private/USDGeomMeshConversion.cpp
@@ -105,6 +105,21 @@ static FAutoConsoleVariableRef CVarMaxSubdivLevel(
 	TEXT("Maximum allowed level of subdivision (1 means a single iteration of subdivision)")
 );
 
+static float GMeshNormalRepairThreshold = 0.05f;
+static FAutoConsoleVariableRef CVarMeshNormalRepairThreshold(
+	TEXT("USD.MeshNormalRepairThreshold"),
+	GMeshNormalRepairThreshold,
+	TEXT("We will try repairing up to this fraction of a Mesh's normals when invalid. If a Mesh has more invalid normals than this, we will "
+		 "recompute all of them. Defaults to 0.05 (5% of all normals).")
+);
+
+static bool GSkipMeshTangentComputation = false;
+static FAutoConsoleVariableRef CVarSkipMeshTangentComputation(
+	TEXT("USD.SkipMeshTangentComputation"),
+	GSkipMeshTangentComputation,
+	TEXT("Skip computing tangents for meshes. With meshes with a huge numer of vertices, it can take a very long time to compute them.")
+);
+
 const FName MeshAttribute::VertexInstance::Velocity("Velocity");
 
 namespace UE::UsdGeomMeshConversion::Private
@@ -2227,6 +2242,7 @@ namespace UE::UsdGeomMeshConversion::Private
 		const int32 VertexInstanceOffset = OutMeshDescription.VertexInstances().Num();
 
 		FStaticMeshAttributes StaticMeshAttributes(OutMeshDescription);
+		StaticMeshAttributes.Register();
 
 		// Vertex positions
 		TVertexAttributesRef<FVector3f> MeshDescriptionVertexPositions = StaticMeshAttributes.GetVertexPositions();
@@ -2291,8 +2307,8 @@ namespace UE::UsdGeomMeshConversion::Private
 					FPolygonGroupID NewPolygonGroup = OutMeshDescription.CreatePolygonGroup();
 					PolygonGroupMapping.Add(CombinedMaterialIndex, NewPolygonGroup);
 
-					// This is important for runtime, where the material slots are matched to LOD sections based on their material slot name
-					MaterialSlotNames[NewPolygonGroup] = *LexToString(NewPolygonGroup.GetValue());
+					const FString& SlotName = InMeshData.LocalMaterialInfo.Slots[LocalMaterialIndex].SlotName;
+					MaterialSlotNames[NewPolygonGroup] = *SlotName;
 				}
 			}
 
@@ -2661,6 +2677,22 @@ bool UsdToUnreal::ConvertGeomMesh(
 	return UsdGeomMeshImpl::ConvertMeshData(MeshData, StageInfo, Options, OutMeshDescription, OutMaterialAssignments);
 }
 
+bool UsdToUnreal::ConvertGeomMesh(
+	const pxr::UsdPrim& MeshPrim,
+	FMeshDescription& InOutMeshDescription,
+	UsdUtils::FUsdPrimMaterialAssignmentInfo& InOutMaterialAssignments,
+	const FUsdMeshConversionOptions& CommonOptions
+)
+{
+	pxr::UsdGeomMesh GeomMesh{MeshPrim};
+	if (!GeomMesh)
+	{
+		return false;
+	}
+
+	return ConvertGeomMesh(GeomMesh, InOutMeshDescription, InOutMaterialAssignments, CommonOptions);
+}
+
 bool UsdToUnreal::ConvertPointInstancerToMesh(
 	const pxr::UsdGeomPointInstancer& PointInstancer,
 	FMeshDescription& OutMeshDescription,
@@ -3321,6 +3353,8 @@ UsdUtils::FUsdPrimMaterialAssignmentInfo UsdUtils::GetPrimMaterialAssignments(
 		return {};
 	}
 
+	FScopedUsdAllocs Allocs;
+
 	auto FetchFirstUEMaterialFromAttribute = [](const pxr::UsdPrim& UsdPrim, const pxr::UsdTimeCode TimeCode) -> TOptional<FString>
 	{
 		FString ValidPackagePath;
@@ -3442,6 +3476,7 @@ UsdUtils::FUsdPrimMaterialAssignmentInfo UsdUtils::GetPrimMaterialAssignments(
 		{
 			const pxr::UsdGeomSubset& GeomSubset = GeomSubsets[GeomSubsetIndex];
 			pxr::UsdPrim GeomSubsetPrim = GeomSubset.GetPrim();
+			FString SlotName = UsdToUnreal::ConvertString(GeomSubsetPrim.GetName());
 			FString GeomSubsetPath = UsdToUnreal::ConvertPath(GeomSubsetPrim.GetPath());
 			bool bHasAssignment = false;
 
@@ -3457,6 +3492,7 @@ UsdUtils::FUsdPrimMaterialAssignmentInfo UsdUtils::GetPrimMaterialAssignments(
 						if (TOptional<FString> UnrealMaterial = UsdUtils::GetUnrealSurfaceOutput(ShadeMaterial.GetPrim()))
 						{
 							FUsdPrimMaterialSlot& Slot = Result.Slots.Emplace_GetRef();
+							Slot.SlotName = SlotName;
 							Slot.MaterialSource = UnrealMaterial.GetValue();
 							Slot.AssignmentType = UsdUtils::EPrimAssignmentType::UnrealMaterial;
 							Slot.bMeshIsDoubleSided = bIsDoubleSided;
@@ -3472,6 +3508,7 @@ UsdUtils::FUsdPrimMaterialAssignmentInfo UsdUtils::GetPrimMaterialAssignments(
 					if (TOptional<FString> UnrealMaterial = FetchFirstUEMaterialFromAttribute(GeomSubsetPrim, TimeCode))
 					{
 						FUsdPrimMaterialSlot& Slot = Result.Slots.Emplace_GetRef();
+						Slot.SlotName = SlotName;
 						Slot.MaterialSource = UnrealMaterial.GetValue();
 						Slot.AssignmentType = UsdUtils::EPrimAssignmentType::UnrealMaterial;
 						Slot.bMeshIsDoubleSided = bIsDoubleSided;
@@ -3487,6 +3524,7 @@ UsdUtils::FUsdPrimMaterialAssignmentInfo UsdUtils::GetPrimMaterialAssignments(
 				if (TOptional<FString> BoundMaterial = FetchMaterialByComputingBoundMaterial(GeomSubsetPrim))
 				{
 					FUsdPrimMaterialSlot& Slot = Result.Slots.Emplace_GetRef();
+					Slot.SlotName = SlotName;
 					Slot.MaterialSource = BoundMaterial.GetValue();
 					Slot.AssignmentType = UsdUtils::EPrimAssignmentType::MaterialPrim;
 					Slot.bMeshIsDoubleSided = bIsDoubleSided;
@@ -3503,6 +3541,7 @@ UsdUtils::FUsdPrimMaterialAssignmentInfo UsdUtils::GetPrimMaterialAssignments(
 			if (!bHasAssignment)
 			{
 				FUsdPrimMaterialSlot& Slot = Result.Slots.Emplace_GetRef();
+				Slot.SlotName = SlotName;
 				Slot.PrimPaths.Add(GeomSubsetPath);
 				bHasAssignment = true;
 			}
@@ -3551,6 +3590,8 @@ UsdUtils::FUsdPrimMaterialAssignmentInfo UsdUtils::GetPrimMaterialAssignments(
 	bool bHasMainAssignment = false;
 	if (bNeedsMainAssignment)
 	{
+		const FString MainSlotName = TEXT("Main");
+
 		// Priority 1: Material is an unreal asset
 		if (RenderContext == UnrealIdentifiers::Unreal)
 		{
@@ -3563,6 +3604,7 @@ UsdUtils::FUsdPrimMaterialAssignmentInfo UsdUtils::GetPrimMaterialAssignments(
 				if (TOptional<FString> UnrealMaterial = UsdUtils::GetUnrealSurfaceOutput(ShadeMaterial.GetPrim()))
 				{
 					FUsdPrimMaterialSlot& Slot = Result.Slots.Emplace_GetRef();
+					Slot.SlotName = MainSlotName;
 					Slot.MaterialSource = UnrealMaterial.GetValue();
 					Slot.AssignmentType = UsdUtils::EPrimAssignmentType::UnrealMaterial;
 					Slot.bMeshIsDoubleSided = bIsDoubleSided;
@@ -3578,6 +3620,7 @@ UsdUtils::FUsdPrimMaterialAssignmentInfo UsdUtils::GetPrimMaterialAssignments(
 				if (TOptional<FString> UnrealMaterial = FetchFirstUEMaterialFromAttribute(UsdPrim, TimeCode))
 				{
 					FUsdPrimMaterialSlot& Slot = Result.Slots.Emplace_GetRef();
+					Slot.SlotName = MainSlotName;
 					Slot.MaterialSource = UnrealMaterial.GetValue();
 					Slot.AssignmentType = UsdUtils::EPrimAssignmentType::UnrealMaterial;
 					Slot.bMeshIsDoubleSided = bIsDoubleSided;
@@ -3594,6 +3637,7 @@ UsdUtils::FUsdPrimMaterialAssignmentInfo UsdUtils::GetPrimMaterialAssignments(
 			if (TOptional<FString> BoundMaterial = FetchMaterialByComputingBoundMaterial(UsdPrim))
 			{
 				FUsdPrimMaterialSlot& Slot = Result.Slots.Emplace_GetRef();
+				Slot.SlotName = MainSlotName;
 				Slot.MaterialSource = BoundMaterial.GetValue();
 				Slot.AssignmentType = UsdUtils::EPrimAssignmentType::MaterialPrim;
 				Slot.bMeshIsDoubleSided = bIsDoubleSided;
@@ -3612,6 +3656,7 @@ UsdUtils::FUsdPrimMaterialAssignmentInfo UsdUtils::GetPrimMaterialAssignments(
 			if (DisplayColor)
 			{
 				FUsdPrimMaterialSlot& Slot = Result.Slots.Emplace_GetRef();
+				Slot.SlotName = MainSlotName;
 				Slot.MaterialSource = DisplayColor.GetValue().ToString();
 				Slot.AssignmentType = UsdUtils::EPrimAssignmentType::DisplayColor;
 				Slot.bMeshIsDoubleSided = bIsDoubleSided;
@@ -4514,6 +4559,74 @@ bool UnrealToUsd::ConvertGeometryCache(const UGeometryCache* GeometryCache, pxr:
 	return true;
 }
 
+void UsdUtils::RepairNormalsAndTangents(const FString& PrimPath, FMeshDescription& MeshDescription)
+{
+	FStaticMeshConstAttributes Attributes{MeshDescription};
+	TArrayView<const FVector3f> VertexInstanceNormals = Attributes.GetVertexInstanceNormals().GetRawArray();
+
+	// Similar to FStaticMeshOperations::AreNormalsAndTangentsValid but we don't care about tangents since we never
+	// read those from USD
+	uint64 InvalidNormalCount = 0;
+	for (const FVertexInstanceID VertexInstanceID : MeshDescription.VertexInstances().GetElementIDs())
+	{
+		if (VertexInstanceNormals[VertexInstanceID].IsNearlyZero() || VertexInstanceNormals[VertexInstanceID].ContainsNaN())
+		{
+			++InvalidNormalCount;
+		}
+	}
+	if (InvalidNormalCount == 0)
+	{
+		return;
+	}
+
+	const float InvalidNormalFraction = (float)InvalidNormalCount / (float)VertexInstanceNormals.Num();
+
+	// We always need to do this at this point as ComputeTangentsAndNormals will end up computing tangents anyway
+	// and our triangle tangents are always invalid
+	FStaticMeshOperations::ComputeTriangleTangentsAndNormals(MeshDescription);
+
+	const static FString MeshNormalRepairThresholdText = TEXT("USD.MeshNormalRepairThreshold");
+
+	// Make sure our normals can be rebuilt from MeshDescription::InitializeAutoGeneratedAttributes in case some tool needs them.
+	// Always force-compute tangents here as we never have them anyway. If we don't force them to be recomputed we'll get
+	// the worst of both worlds as some of these will be arbitrarily recomputed anyway, and some will be left invalid
+	EComputeNTBsFlags Options = GSkipMeshTangentComputation ? EComputeNTBsFlags::None
+															: EComputeNTBsFlags::UseMikkTSpace | EComputeNTBsFlags::Tangents;
+
+	// Repairing can take a long time for degenerate triangles (UE-194839)
+	Options |= EComputeNTBsFlags::IgnoreDegenerateTriangles;
+
+	if (InvalidNormalFraction >= GMeshNormalRepairThreshold)
+	{
+		Options |= EComputeNTBsFlags::Normals;
+		UE_LOG(
+			LogUsd,
+			Verbose,
+			TEXT("%f%% of the normals from Mesh prim '%s' are invalid or unusable. This is at or above the threshold of '%f%%' (configurable via "
+				 "the cvar '%s'), so normals will be discarded and fully recomputed. Note that when the cvar "
+				 "'USD.Subdiv.IgnoreNormalsWhenSubdividing' is true it is expected for subdivision meshes to have their normals discarded."),
+			InvalidNormalFraction * 100.0f,
+			*PrimPath,
+			GMeshNormalRepairThreshold * 100.0f,
+			*MeshNormalRepairThresholdText
+		);
+	}
+	else if (InvalidNormalFraction > 0)
+	{
+		UE_LOG(
+			LogUsd,
+			Verbose,
+			TEXT("%f%% of the normals from Mesh prim '%s' are invalid or unusable. This is below the threshold of '%f%%' (configurable via the "
+				 "cvar '%s'), so the invalid normals will be repaired."),
+			InvalidNormalFraction * 100.0f,
+			*PrimPath,
+			GMeshNormalRepairThreshold * 100.0f,
+			*MeshNormalRepairThresholdText
+		);
+	}
+	FStaticMeshOperations::ComputeTangentsAndNormals(MeshDescription, Options);
+}
+
 TOptional<UsdUtils::FDisplayColorMaterial> UsdUtils::ExtractDisplayColorMaterial(const pxr::UsdGeomGprim& Gprim, const pxr::UsdTimeCode TimeCode)
 {
 	if (!Gprim)
diff --git a/Engine/Plugins/Runtime/USDCore/Source/USDUtilities/Private/USDShadeConversion.cpp b/Engine/Plugins/Runtime/USDCore/Source/USDUtilities/Private/USDShadeConversion.cpp
--- a/Engine/Plugins/Runtime/USDCore/Source/USDUtilities/Private/USDShadeConversion.cpp
+++ b/Engine/Plugins/Runtime/USDCore/Source/USDUtilities/Private/USDShadeConversion.cpp
@@ -66,2101 +66,2092 @@ class UTextureFactory;
 
 #define LOCTEXT_NAMESPACE "USDShadeConversion"
 
-namespace UE
+namespace UE::UsdShadeConversion::Private
 {
-	namespace UsdShadeConversion
-	{
-		namespace Private
-		{
 #if WITH_EDITOR
-			using FlattenToMatEntry = TPairInitializer<const EFlattenMaterialProperties&, const EMaterialProperty&>;
-			const static TMap<EFlattenMaterialProperties, EMaterialProperty> FlattenToMaterialProperty({
-				FlattenToMatEntry(EFlattenMaterialProperties::Diffuse, EMaterialProperty::MP_BaseColor),
-				FlattenToMatEntry(EFlattenMaterialProperties::Metallic, EMaterialProperty::MP_Metallic),
-				FlattenToMatEntry(EFlattenMaterialProperties::Specular, EMaterialProperty::MP_Specular),
-				FlattenToMatEntry(EFlattenMaterialProperties::Roughness, EMaterialProperty::MP_Roughness),
-				FlattenToMatEntry(EFlattenMaterialProperties::Anisotropy, EMaterialProperty::MP_Anisotropy),
-				FlattenToMatEntry(EFlattenMaterialProperties::Normal, EMaterialProperty::MP_Normal),
-				FlattenToMatEntry(EFlattenMaterialProperties::Tangent, EMaterialProperty::MP_Tangent),
-				FlattenToMatEntry(EFlattenMaterialProperties::Opacity, EMaterialProperty::MP_Opacity),
-				FlattenToMatEntry(EFlattenMaterialProperties::Emissive, EMaterialProperty::MP_EmissiveColor),
-				FlattenToMatEntry(EFlattenMaterialProperties::SubSurface, EMaterialProperty::MP_SubsurfaceColor),
-				FlattenToMatEntry(EFlattenMaterialProperties::OpacityMask, EMaterialProperty::MP_OpacityMask),
-				FlattenToMatEntry(EFlattenMaterialProperties::AmbientOcclusion, EMaterialProperty::MP_AmbientOcclusion),
-			});
-
-			/** Simple wrapper around the FBakeOutput data that we can reuse for data coming in from a FFlattenMaterial */
-			struct FBakedMaterialView
-			{
-				TMap<EMaterialProperty, TArray<FColor>*> PropertyData;
-				TMap<EMaterialProperty, TArray<FFloat16Color>*> HDRPropertyData;
-				TMap<EMaterialProperty, FIntPoint> PropertySizes;
-
-				explicit FBakedMaterialView(FBakeOutput& BakeOutput)
-					: PropertySizes(BakeOutput.PropertySizes)
-				{
-					HDRPropertyData.Reserve(BakeOutput.HDRPropertyData.Num());
-					for (TPair<EMaterialProperty, TArray<FFloat16Color>>& BakedPair : BakeOutput.HDRPropertyData)
-					{
-						HDRPropertyData.Add(BakedPair.Key, &BakedPair.Value);
-					}
-
-					PropertyData.Reserve(BakeOutput.PropertyData.Num());
-					for (TPair<EMaterialProperty, TArray<FColor>>& BakedPair : BakeOutput.PropertyData)
-					{
-						PropertyData.Add(BakedPair.Key, &BakedPair.Value);
-					}
-				}
+	using FlattenToMatEntry = TPairInitializer<const EFlattenMaterialProperties&, const EMaterialProperty&>;
+	const static TMap<EFlattenMaterialProperties, EMaterialProperty> FlattenToMaterialProperty({
+		FlattenToMatEntry(EFlattenMaterialProperties::Diffuse, EMaterialProperty::MP_BaseColor),
+		FlattenToMatEntry(EFlattenMaterialProperties::Metallic, EMaterialProperty::MP_Metallic),
+		FlattenToMatEntry(EFlattenMaterialProperties::Specular, EMaterialProperty::MP_Specular),
+		FlattenToMatEntry(EFlattenMaterialProperties::Roughness, EMaterialProperty::MP_Roughness),
+		FlattenToMatEntry(EFlattenMaterialProperties::Anisotropy, EMaterialProperty::MP_Anisotropy),
+		FlattenToMatEntry(EFlattenMaterialProperties::Normal, EMaterialProperty::MP_Normal),
+		FlattenToMatEntry(EFlattenMaterialProperties::Tangent, EMaterialProperty::MP_Tangent),
+		FlattenToMatEntry(EFlattenMaterialProperties::Opacity, EMaterialProperty::MP_Opacity),
+		FlattenToMatEntry(EFlattenMaterialProperties::Emissive, EMaterialProperty::MP_EmissiveColor),
+		FlattenToMatEntry(EFlattenMaterialProperties::SubSurface, EMaterialProperty::MP_SubsurfaceColor),
+		FlattenToMatEntry(EFlattenMaterialProperties::OpacityMask, EMaterialProperty::MP_OpacityMask),
+		FlattenToMatEntry(EFlattenMaterialProperties::AmbientOcclusion, EMaterialProperty::MP_AmbientOcclusion),
+	});
+
+	/** Simple wrapper around the FBakeOutput data that we can reuse for data coming in from a FFlattenMaterial */
+	struct FBakedMaterialView
+	{
+		TMap<EMaterialProperty, TArray<FColor>*> PropertyData;
+		TMap<EMaterialProperty, TArray<FFloat16Color>*> HDRPropertyData;
+		TMap<EMaterialProperty, FIntPoint> PropertySizes;
+
+		explicit FBakedMaterialView(FBakeOutput& BakeOutput)
+			: PropertySizes(BakeOutput.PropertySizes)
+		{
+			HDRPropertyData.Reserve(BakeOutput.HDRPropertyData.Num());
+			for (TPair<EMaterialProperty, TArray<FFloat16Color>>& BakedPair : BakeOutput.HDRPropertyData)
+			{
+				HDRPropertyData.Add(BakedPair.Key, &BakedPair.Value);
+			}
 
-				explicit FBakedMaterialView(FFlattenMaterial& FlattenMaterial)
-				{
-					PropertyData.Reserve(static_cast<uint8>(EFlattenMaterialProperties::NumFlattenMaterialProperties));
-					for (const TPair<EFlattenMaterialProperties, EMaterialProperty>& FlattenToProperty : FlattenToMaterialProperty)
-					{
-						PropertyData.Add(FlattenToProperty.Value, &FlattenMaterial.GetPropertySamples(FlattenToProperty.Key));
-					}
+			PropertyData.Reserve(BakeOutput.PropertyData.Num());
+			for (TPair<EMaterialProperty, TArray<FColor>>& BakedPair : BakeOutput.PropertyData)
+			{
+				PropertyData.Add(BakedPair.Key, &BakedPair.Value);
+			}
+		}
 
-					PropertySizes.Reserve(static_cast<uint8>(EFlattenMaterialProperties::NumFlattenMaterialProperties));
-					for (const TPair<EFlattenMaterialProperties, EMaterialProperty>& FlattenToProperty : FlattenToMaterialProperty)
-					{
-						PropertySizes.Add(FlattenToProperty.Value, FlattenMaterial.GetPropertySize(FlattenToProperty.Key));
-					}
-				}
-			};
-#endif	  // WITH_EDITOR
+		explicit FBakedMaterialView(FFlattenMaterial& FlattenMaterial)
+		{
+			PropertyData.Reserve(static_cast<uint8>(EFlattenMaterialProperties::NumFlattenMaterialProperties));
+			for (const TPair<EFlattenMaterialProperties, EMaterialProperty>& FlattenToProperty : FlattenToMaterialProperty)
+			{
+				PropertyData.Add(FlattenToProperty.Value, &FlattenMaterial.GetPropertySamples(FlattenToProperty.Key));
+			}
 
-			// Given an AssetPath, resolve it to an actual file path
-			FString ResolveAssetPath(const pxr::SdfLayerHandle& LayerHandle, const FString& AssetPath)
+			PropertySizes.Reserve(static_cast<uint8>(EFlattenMaterialProperties::NumFlattenMaterialProperties));
+			for (const TPair<EFlattenMaterialProperties, EMaterialProperty>& FlattenToProperty : FlattenToMaterialProperty)
 			{
-				// TODO: Most of this stuff is incompatible with custom resolvers, as these asset paths may be URLs
-				// or just GUIDs or anything like that, where relative vs absolute path make no sense. We will need
-				// a different approach whenever we properly handle USD resolvers
+				PropertySizes.Add(FlattenToProperty.Value, FlattenMaterial.GetPropertySize(FlattenToProperty.Key));
+			}
+		}
+	};
+#endif	  // WITH_EDITOR
 
-				FScopedUsdAllocs UsdAllocs;
+	// Given an AssetPath, resolve it to an actual file path
+	FString ResolveAssetPath(const pxr::SdfLayerHandle& LayerHandle, const FString& AssetPath)
+	{
+		// TODO: Most of this stuff is incompatible with custom resolvers, as these asset paths may be URLs
+		// or just GUIDs or anything like that, where relative vs absolute path make no sense. We will need
+		// a different approach whenever we properly handle USD resolvers
 
-				FString PathToResolve = LayerHandle ? UsdToUnreal::ConvertString(
-											pxr::SdfComputeAssetPathRelativeToLayer(LayerHandle, UnrealToUsd::ConvertString(*AssetPath).Get())
-										)
-													: AssetPath;
+		FScopedUsdAllocs UsdAllocs;
 
-				// We need to provide absolute paths to the resolver later: It has no idea what to do with a relative
-				// path relative to some random location
-				if (FPaths::IsRelative(PathToResolve) && LayerHandle)
-				{
-					FString LayerDirectory = FPaths::GetPath(UsdToUnreal::ConvertString(LayerHandle->GetRealPath()));
-					PathToResolve = FPaths::Combine(LayerDirectory, PathToResolve);
-				}
+		FString PathToResolve = LayerHandle ? UsdToUnreal::ConvertString(
+									pxr::SdfComputeAssetPathRelativeToLayer(LayerHandle, UnrealToUsd::ConvertString(*AssetPath).Get())
+								)
+											: AssetPath;
 
-				// If this path has an UDIM placeholder in it (e.g. "textures/red.<UDIM>.exr"), let's try to find an
-				// actual existing UDIM tile instead, or else Resolver.Resolve will just give us the empty string.
-				// Previously we directly tried to find tile 1001, but it seems there is no guarantee that the
-				// particular 1001 tile will exist.
-				// Also note that for UDIMs the UE texture factory expects to receive path to any one tile, and it
-				// will itself discover the remaining tiles (need to also keep that in mind when doing the support
-				// for custom resolvers...)
-				if (PathToResolve.Contains(TEXT("<UDIM>")))
-				{
-					FString UdimFileFilter = PathToResolve.Replace(TEXT("<UDIM>"), TEXT("*"));
+		// We need to provide absolute paths to the resolver later: It has no idea what to do with a relative
+		// path relative to some random location
+		if (FPaths::IsRelative(PathToResolve) && LayerHandle)
+		{
+			FString LayerDirectory = FPaths::GetPath(UsdToUnreal::ConvertString(LayerHandle->GetRealPath()));
+			PathToResolve = FPaths::Combine(LayerDirectory, PathToResolve);
+		}
 
-					TArray<FString> UDIMFiles;
-					IFileManager::Get().FindFiles(UDIMFiles, *UdimFileFilter, /*files*/ true, /*directories*/ false);
-					if (UDIMFiles.Num() > 0)
-					{
-						// Sort here to enforce some sort of consistency between repeated calls
-						UDIMFiles.Sort();
+		// If this path has an UDIM placeholder in it (e.g. "textures/red.<UDIM>.exr"), let's try to find an
+		// actual existing UDIM tile instead, or else Resolver.Resolve will just give us the empty string.
+		// Previously we directly tried to find tile 1001, but it seems there is no guarantee that the
+		// particular 1001 tile will exist.
+		// Also note that for UDIMs the UE texture factory expects to receive path to any one tile, and it
+		// will itself discover the remaining tiles (need to also keep that in mind when doing the support
+		// for custom resolvers...)
+		if (PathToResolve.Contains(TEXT("<UDIM>")))
+		{
+			FString UdimFileFilter = PathToResolve.Replace(TEXT("<UDIM>"), TEXT("*"));
 
-						// FindFiles will just return the filename with no path info, so lets put the file in the same
-						// location our original <UDIM> path was
-						PathToResolve = FPaths::Combine(FPaths::GetPath(PathToResolve), UDIMFiles[0]);
-					}
-				}
+			TArray<FString> UDIMFiles;
+			IFileManager::Get().FindFiles(UDIMFiles, *UdimFileFilter, /*files*/ true, /*directories*/ false);
+			if (UDIMFiles.Num() > 0)
+			{
+				// Sort here to enforce some sort of consistency between repeated calls
+				UDIMFiles.Sort();
 
-				pxr::ArResolverScopedCache ResolverCache;
-				pxr::ArResolver& Resolver = pxr::ArGetResolver();
-				std::string AssetIdentifier = Resolver.Resolve(UnrealToUsd::ConvertString(*PathToResolve).Get().c_str());
+				// FindFiles will just return the filename with no path info, so lets put the file in the same
+				// location our original <UDIM> path was
+				PathToResolve = FPaths::Combine(FPaths::GetPath(PathToResolve), UDIMFiles[0]);
+			}
+		}
 
-				// Don't normalize an empty path as the result will be "."
-				if (AssetIdentifier.size() > 0)
-				{
-					AssetIdentifier = Resolver.CreateIdentifier(AssetIdentifier);
-				}
+		pxr::ArResolverScopedCache ResolverCache;
+		pxr::ArResolver& Resolver = pxr::ArGetResolver();
+		std::string AssetIdentifier = Resolver.Resolve(UnrealToUsd::ConvertString(*PathToResolve).Get().c_str());
 
-				FString ResolvedAssetPath = UsdToUnreal::ConvertString(AssetIdentifier);
-				return ResolvedAssetPath;
-			}
+		// Don't normalize an empty path as the result will be "."
+		if (AssetIdentifier.size() > 0)
+		{
+			AssetIdentifier = Resolver.CreateIdentifier(AssetIdentifier);
+		}
 
-			// If ResolvedTexturePath is in a format like "C:/MyFiles/scene.usdz[0/texture.png]", this will place "png" in OutExtensionNoDot, and
-			// return true
-			bool IsInsideUsdzArchive(const FString& ResolvedTexturePath, FString& OutExtensionNoDot)
-			{
-				// We need at least an opening and closing bracket
-				if (ResolvedTexturePath.Len() < 3)
-				{
-					return false;
-				}
+		FString ResolvedAssetPath = UsdToUnreal::ConvertString(AssetIdentifier);
+		return ResolvedAssetPath;
+	}
 
-				int32 OpenBracketPos = ResolvedTexturePath.Find(TEXT("["), ESearchCase::IgnoreCase, ESearchDir::FromEnd);
-				int32 CloseBracketPos = ResolvedTexturePath.Find(TEXT("]"), ESearchCase::IgnoreCase, ESearchDir::FromEnd);
-				if (OpenBracketPos != INDEX_NONE && CloseBracketPos != INDEX_NONE && CloseBracketPos > OpenBracketPos - 1)
-				{
-					// Should be like "C:/MyFiles/scene.usdz"
-					FString UsdzFilePath = ResolvedTexturePath.Left(OpenBracketPos);
+	// If ResolvedTexturePath is in a format like "C:/MyFiles/scene.usdz[0/texture.png]", this will place "png" in OutExtensionNoDot, and
+	// return true
+	bool IsInsideUsdzArchive(const FString& ResolvedTexturePath, FString& OutExtensionNoDot)
+	{
+		// We need at least an opening and closing bracket
+		if (ResolvedTexturePath.Len() < 3)
+		{
+			return false;
+		}
 
-					// Should be like "0/texture.png"
-					FString TextureInUsdzPath = ResolvedTexturePath.Mid(OpenBracketPos + 1, CloseBracketPos - OpenBracketPos - 1);
+		int32 OpenBracketPos = ResolvedTexturePath.Find(TEXT("["), ESearchCase::IgnoreCase, ESearchDir::FromEnd);
+		int32 CloseBracketPos = ResolvedTexturePath.Find(TEXT("]"), ESearchCase::IgnoreCase, ESearchDir::FromEnd);
+		if (OpenBracketPos != INDEX_NONE && CloseBracketPos != INDEX_NONE && CloseBracketPos > OpenBracketPos - 1)
+		{
+			// Should be like "C:/MyFiles/scene.usdz"
+			FString UsdzFilePath = ResolvedTexturePath.Left(OpenBracketPos);
 
-					if (FPaths::GetExtension(UsdzFilePath).ToLower() == TEXT("usdz"))
-					{
-						// Should be something like "png"
-						OutExtensionNoDot = FPaths::GetExtension(TextureInUsdzPath);
-						return true;
-					}
-				}
+			// Should be like "0/texture.png"
+			FString TextureInUsdzPath = ResolvedTexturePath.Mid(OpenBracketPos + 1, CloseBracketPos - OpenBracketPos - 1);
 
-				return false;
+			if (FPaths::GetExtension(UsdzFilePath).ToLower() == TEXT("usdz"))
+			{
+				// Should be something like "png"
+				OutExtensionNoDot = FPaths::GetExtension(TextureInUsdzPath);
+				return true;
 			}
+		}
 
-			TUsdStore<std::shared_ptr<const char>> ReadTextureBufferFromUsdzArchive(const FString& ResolvedTexturePath, uint64& OutBufferSize)
-			{
-				pxr::ArResolver& Resolver = pxr::ArGetResolver();
-				std::shared_ptr<pxr::ArAsset> Asset = Resolver.OpenAsset(pxr::ArResolvedPath(UnrealToUsd::ConvertString(*ResolvedTexturePath).Get()));
+		return false;
+	}
 
-				TUsdStore<std::shared_ptr<const char>> Buffer;
+	TUsdStore<std::shared_ptr<const char>> ReadTextureBufferFromUsdzArchive(const FString& ResolvedTexturePath, uint64& OutBufferSize)
+	{
+		pxr::ArResolver& Resolver = pxr::ArGetResolver();
+		std::shared_ptr<pxr::ArAsset> Asset = Resolver.OpenAsset(pxr::ArResolvedPath(UnrealToUsd::ConvertString(*ResolvedTexturePath).Get()));
 
-				if (Asset)
-				{
-					OutBufferSize = static_cast<uint64>(Asset->GetSize());
-					{
-						FScopedUsdAllocs Allocs;
-						Buffer = Asset->GetBuffer();
-					}
-				}
+		TUsdStore<std::shared_ptr<const char>> Buffer;
 
-				return Buffer;
+		if (Asset)
+		{
+			OutBufferSize = static_cast<uint64>(Asset->GetSize());
+			{
+				FScopedUsdAllocs Allocs;
+				Buffer = Asset->GetBuffer();
 			}
+		}
 
-			/** If ResolvedTexturePath points at a texture inside an usdz file, this will use USD to pull the asset from the file, and TextureFactory
-			 * to import it directly from the binary buffer */
-			UTexture* ReadTextureFromUsdzArchiveEditor(
-				const FString& ResolvedTexturePath,
-				const FString& TextureExtension,
-				UTextureFactory* TextureFactory,
-				UObject* Outer,
-				EObjectFlags ObjectFlags
-			)
-			{
+		return Buffer;
+	}
+
+	/** If ResolvedTexturePath points at a texture inside an usdz file, this will use USD to pull the asset from the file, and TextureFactory
+	 * to import it directly from the binary buffer */
+	UTexture* ReadTextureFromUsdzArchiveEditor(
+		const FString& ResolvedTexturePath,
+		const FString& TextureExtension,
+		UTextureFactory* TextureFactory,
+		UObject* Outer,
+		EObjectFlags ObjectFlags
+	)
+	{
 #if WITH_EDITOR
-				uint64 BufferSize = 0;
-				TUsdStore<std::shared_ptr<const char>> Buffer = ReadTextureBufferFromUsdzArchive(ResolvedTexturePath, BufferSize);
-				const uint8* BufferStart = reinterpret_cast<const uint8*>(Buffer.Get().get());
+		uint64 BufferSize = 0;
+		TUsdStore<std::shared_ptr<const char>> Buffer = ReadTextureBufferFromUsdzArchive(ResolvedTexturePath, BufferSize);
+		const uint8* BufferStart = reinterpret_cast<const uint8*>(Buffer.Get().get());
 
-				if (BufferSize == 0 || BufferStart == nullptr)
-				{
-					return nullptr;
-				}
+		if (BufferSize == 0 || BufferStart == nullptr)
+		{
+			return nullptr;
+		}
 
-				FScopedUnrealAllocs UEAllocs;
+		FScopedUnrealAllocs UEAllocs;
 
-				FName TextureName = MakeUniqueObjectName(
-					Outer,
-					UTexture::StaticClass(),
-					*UsdUnreal::ObjectUtils::SanitizeObjectName(FPaths::GetBaseFilename(ResolvedTexturePath))
-				);
+		FName TextureName = MakeUniqueObjectName(
+			Outer,
+			UTexture::StaticClass(),
+			*UsdUnreal::ObjectUtils::SanitizeObjectName(FPaths::GetBaseFilename(ResolvedTexturePath))
+		);
 
-				UObject* Context = nullptr;
-				FFeedbackContext* FeedbackContext = GWarn;
-				return Cast<UTexture>(TextureFactory->FactoryCreateBinary(
-					UTexture::StaticClass(),
-					Outer,
-					TextureName,
-					ObjectFlags,
-					Context,
-					*TextureExtension,
-					BufferStart,
-					BufferStart + BufferSize,
-					FeedbackContext
-				));
+		UObject* Context = nullptr;
+		FFeedbackContext* FeedbackContext = GWarn;
+		return Cast<UTexture>(TextureFactory->FactoryCreateBinary(
+			UTexture::StaticClass(),
+			Outer,
+			TextureName,
+			ObjectFlags,
+			Context,
+			*TextureExtension,
+			BufferStart,
+			BufferStart + BufferSize,
+			FeedbackContext
+		));
 #endif	  // WITH_EDITOR
-				return nullptr;
-			}
-
-			UTexture* ReadTextureFromUsdzArchiveRuntime(const FString& ResolvedTexturePath)
-			{
-				uint64 BufferSize = 0;
-				TUsdStore<std::shared_ptr<const char>> Buffer = ReadTextureBufferFromUsdzArchive(ResolvedTexturePath, BufferSize);
-				const uint8* BufferStart = reinterpret_cast<const uint8*>(Buffer.Get().get());
-
-				if (BufferSize == 0 || BufferStart == nullptr)
-				{
-					return nullptr;
-				}
-
-				FScopedUnrealAllocs UEAllocs;
-
-				// Copied from FImageUtils::ImportBufferAsTexture2D( Buffer ) so that we can avoid a copy into the TArray<uint8> that it takes as
-				// parameter
-				IImageWrapperModule& ImageWrapperModule = FModuleManager::Get().LoadModuleChecked<IImageWrapperModule>(TEXT("ImageWrapper"));
-				EImageFormat Format = ImageWrapperModule.DetectImageFormat(BufferStart, BufferSize);
-
-				UTexture2D* NewTexture = nullptr;
-				EPixelFormat PixelFormat = PF_Unknown;
-				if (Format != EImageFormat::Invalid)
-				{
-					TSharedPtr<IImageWrapper> ImageWrapper = ImageWrapperModule.CreateImageWrapper(Format);
+		return nullptr;
+	}
 
-					int32 BitDepth = 0;
-					int32 Width = 0;
-					int32 Height = 0;
+	UTexture* ReadTextureFromUsdzArchiveRuntime(const FString& ResolvedTexturePath)
+	{
+		uint64 BufferSize = 0;
+		TUsdStore<std::shared_ptr<const char>> Buffer = ReadTextureBufferFromUsdzArchive(ResolvedTexturePath, BufferSize);
+		const uint8* BufferStart = reinterpret_cast<const uint8*>(Buffer.Get().get());
 
-					if (ImageWrapper->SetCompressed((void*)BufferStart, BufferSize))
-					{
-						PixelFormat = PF_Unknown;
+		if (BufferSize == 0 || BufferStart == nullptr)
+		{
+			return nullptr;
+		}
 
-						ERGBFormat RGBFormat = ERGBFormat::Invalid;
+		FScopedUnrealAllocs UEAllocs;
 
-						BitDepth = ImageWrapper->GetBitDepth();
+		// Copied from FImageUtils::ImportBufferAsTexture2D( Buffer ) so that we can avoid a copy into the TArray<uint8> that it takes as
+		// parameter
+		IImageWrapperModule& ImageWrapperModule = FModuleManager::Get().LoadModuleChecked<IImageWrapperModule>(TEXT("ImageWrapper"));
+		EImageFormat Format = ImageWrapperModule.DetectImageFormat(BufferStart, BufferSize);
 
-						Width = ImageWrapper->GetWidth();
-						Height = ImageWrapper->GetHeight();
+		UTexture2D* NewTexture = nullptr;
+		EPixelFormat PixelFormat = PF_Unknown;
+		if (Format != EImageFormat::Invalid)
+		{
+			TSharedPtr<IImageWrapper> ImageWrapper = ImageWrapperModule.CreateImageWrapper(Format);
 
-						if (BitDepth == 16)
-						{
-							PixelFormat = PF_FloatRGBA;
-							RGBFormat = ERGBFormat::BGRA;
-						}
-						else if (BitDepth == 8)
-						{
-							PixelFormat = PF_B8G8R8A8;
-							RGBFormat = ERGBFormat::BGRA;
-						}
-						else
-						{
-							UE_LOG(LogUsd, Warning, TEXT("Error creating texture. Bit depth is unsupported. (%d)"), BitDepth);
-							return nullptr;
-						}
+			int32 BitDepth = 0;
+			int32 Width = 0;
+			int32 Height = 0;
 
-						TArray64<uint8> UncompressedData;
-						ImageWrapper->GetRaw(RGBFormat, BitDepth, UncompressedData);
+			if (ImageWrapper->SetCompressed((void*)BufferStart, BufferSize))
+			{
+				PixelFormat = PF_Unknown;
 
-						NewTexture = UTexture2D::CreateTransient(Width, Height, PixelFormat);
-						if (NewTexture)
-						{
-							NewTexture->bNotOfflineProcessed = true;
-							uint8* MipData = static_cast<uint8*>(NewTexture->GetPlatformData()->Mips[0].BulkData.Lock(LOCK_READ_WRITE));
+				ERGBFormat RGBFormat = ERGBFormat::Invalid;
 
-							// Bulk data was already allocated for the correct size when we called CreateTransient above
-							FMemory::Memcpy(MipData, UncompressedData.GetData(), NewTexture->GetPlatformData()->Mips[0].BulkData.GetBulkDataSize());
+				BitDepth = ImageWrapper->GetBitDepth();
 
-							NewTexture->GetPlatformData()->Mips[0].BulkData.Unlock();
+				Width = ImageWrapper->GetWidth();
+				Height = ImageWrapper->GetHeight();
 
-							NewTexture->UpdateResource();
-						}
-					}
+				if (BitDepth == 16)
+				{
+					PixelFormat = PF_FloatRGBA;
+					RGBFormat = ERGBFormat::BGRA;
+				}
+				else if (BitDepth == 8)
+				{
+					PixelFormat = PF_B8G8R8A8;
+					RGBFormat = ERGBFormat::BGRA;
 				}
 				else
 				{
-					UE_LOG(LogUsd, Warning, TEXT("Error creating texture. Couldn't determine the file format"));
+					UE_LOG(LogUsd, Warning, TEXT("Error creating texture. Bit depth is unsupported. (%d)"), BitDepth);
+					return nullptr;
 				}
 
-				return NewTexture;
-			}
+				TArray64<uint8> UncompressedData;
+				ImageWrapper->GetRaw(RGBFormat, BitDepth, UncompressedData);
 
-			// Will traverse the shade material graph backwards looking for a string/token value and return it.
-			// Returns the empty string if it didn't find anything.
-			FString RecursivelySearchForStringValue(pxr::UsdShadeInput Input)
-			{
-				if (!Input)
+				NewTexture = UTexture2D::CreateTransient(Width, Height, PixelFormat);
+				if (NewTexture)
 				{
-					return {};
-				}
+					NewTexture->bNotOfflineProcessed = true;
+					uint8* MipData = static_cast<uint8*>(NewTexture->GetPlatformData()->Mips[0].BulkData.Lock(LOCK_READ_WRITE));
 
-				FScopedUsdAllocs Allocs;
+					// Bulk data was already allocated for the correct size when we called CreateTransient above
+					FMemory::Memcpy(MipData, UncompressedData.GetData(), NewTexture->GetPlatformData()->Mips[0].BulkData.GetBulkDataSize());
 
-				if (Input.HasConnectedSource())
-				{
-					pxr::UsdShadeConnectableAPI Souce;
-					pxr::TfToken SourceName;
-					pxr::UsdShadeAttributeType SourceType;
-					if (pxr::UsdShadeConnectableAPI::GetConnectedSource(Input.GetAttr(), &Souce, &SourceName, &SourceType))
-					{
-						for (const pxr::UsdShadeInput& DeeperInput : Souce.GetInputs())
-						{
-							FString Result = RecursivelySearchForStringValue(DeeperInput);
-							if (!Result.IsEmpty())
-							{
-								return Result;
-							}
-						}
-					}
-				}
-				else
-				{
-					std::string StringValue;
-					if (Input.Get<std::string>(&StringValue))
-					{
-						return UsdToUnreal::ConvertString(StringValue);
-					}
+					NewTexture->GetPlatformData()->Mips[0].BulkData.Unlock();
 
-					pxr::TfToken TokenValue;
-					if (Input.Get<pxr::TfToken>(&TokenValue))
-					{
-						return UsdToUnreal::ConvertToken(TokenValue);
-					}
+					NewTexture->UpdateResource();
 				}
-
-				return {};
 			}
+		}
+		else
+		{
+			UE_LOG(LogUsd, Warning, TEXT("Error creating texture. Couldn't determine the file format"));
+		}
 
-			struct FTextureParameterValue
-			{
-				UTexture* Texture;
-				FString Primvar;
-				FScale2f UVScale;
-				float UVRotation;
-				FVector2f UVTranslation;
-				int32 OutputIndex = 0;
-			};
+		return NewTexture;
+	}
 
-			struct FPrimvarReaderParameterValue
-			{
-				std::string PrimvarName;
-				FVector FallbackValue;
-			};
+	// Will traverse the shade material graph backwards looking for a string/token value and return it.
+	// Returns the empty string if it didn't find anything.
+	FString RecursivelySearchForStringValue(pxr::UsdShadeInput Input)
+	{
+		if (!Input)
+		{
+			return {};
+		}
 
-			using FParameterValue = TVariant<float, FVector, FTextureParameterValue, FPrimvarReaderParameterValue, bool>;
+		FScopedUsdAllocs Allocs;
 
-			/**
-			 * Receives a UsdUVTexture shader, and returns the name of the primvar that it is using as 'st',
-			 * plus the USD-space UV transforms that should be applied to that primvar when sampling this texture with
-			 * it.
-			 */
-			void GetSTPrimvarAndTransform(
-				const pxr::UsdShadeConnectableAPI& UsdUVTexture,
-				UE::UsdShadeConversion::Private::FTextureParameterValue& OutTextureValue
-			)
+		if (Input.HasConnectedSource())
+		{
+			pxr::UsdShadeConnectableAPI Souce;
+			pxr::TfToken SourceName;
+			pxr::UsdShadeAttributeType SourceType;
+			if (pxr::UsdShadeConnectableAPI::GetConnectedSource(Input.GetAttr(), &Souce, &SourceName, &SourceType))
 			{
-				FScopedUsdAllocs Allocs;
-
-				pxr::UsdShadeInput StInput = UsdUVTexture.GetInput(UnrealIdentifiers::St);
-				if (!StInput)
+				for (const pxr::UsdShadeInput& DeeperInput : Souce.GetInputs())
 				{
-					return;
+					FString Result = RecursivelySearchForStringValue(DeeperInput);
+					if (!Result.IsEmpty())
+					{
+						return Result;
+					}
 				}
+			}
+		}
+		else
+		{
+			std::string StringValue;
+			if (Input.Get<std::string>(&StringValue))
+			{
+				return UsdToUnreal::ConvertString(StringValue);
+			}
 
-				pxr::UsdShadeConnectableAPI Connectable;
-				pxr::TfToken ConnectableOutput;
-				pxr::UsdShadeAttributeType AttributeType;
+			pxr::TfToken TokenValue;
+			if (Input.Get<pxr::TfToken>(&TokenValue))
+			{
+				return UsdToUnreal::ConvertToken(TokenValue);
+			}
+		}
 
-				FTransform2f ConcatenatedTransform;
-				bool bFoundPrimvarReader = false;
+		return {};
+	}
 
-				// Traverse through potentially N UV transform nodes
-				while (pxr::UsdShadeConnectableAPI::GetConnectedSource(StInput.GetAttr(), &Connectable, &ConnectableOutput, &AttributeType))
-				{
-					pxr::UsdPrim ConnectedPrim = Connectable.GetPrim();
-					pxr::UsdShadeNodeDefAPI ConnectedShadeNode{ConnectedPrim};
-					if (!ConnectedShadeNode)
-					{
-						// Deadend, should never really happen
-						return;
-					}
+	/**
+	 * Receives a UsdUVTexture shader, and returns the name of the primvar that it is using as 'st',
+	 * plus the USD-space UV transforms that should be applied to that primvar when sampling this texture with
+	 * it.
+	 */
+	void GetSTPrimvarAndTransform(const pxr::UsdShadeConnectableAPI& UsdUVTexture, UsdToUnreal::FTextureParameterValue& OutTextureValue)
+	{
+		FScopedUsdAllocs Allocs;
 
-					pxr::TfToken ConnectedNodeId;
-					ConnectedShadeNode.GetShaderId(&ConnectedNodeId);
+		pxr::UsdShadeInput StInput = UsdUVTexture.GetInput(UnrealIdentifiers::St);
+		if (!StInput)
+		{
+			return;
+		}
 
-					// UV transform shader node
-					if (ConnectedNodeId == UnrealIdentifiers::UsdTransform2d)
-					{
-						FScale2f Scale{1.0f, 1.0f};
-						if (pxr::UsdShadeInput ScaleInput = Connectable.GetInput(UnrealIdentifiers::Scale))
-						{
-							pxr::GfVec2f VecValue;
-							if (ScaleInput.Get(&VecValue))
-							{
-								Scale = FScale2f{VecValue[0], VecValue[1]};
-							}
-						}
+		pxr::UsdShadeConnectableAPI Connectable;
+		pxr::TfToken ConnectableOutput;
+		pxr::UsdShadeAttributeType AttributeType;
 
-						float Rotation = 0.0f;
-						if (pxr::UsdShadeInput RotationInput = Connectable.GetInput(UnrealIdentifiers::Rotation))
-						{
-							float FloatValue;
-							if (RotationInput.Get(&FloatValue))
-							{
-								Rotation = FloatValue;
-							}
-						}
+		FTransform2f ConcatenatedTransform;
+		bool bFoundPrimvarReader = false;
 
-						FVector2f Translation{0.0f, 0.0f};
-						if (pxr::UsdShadeInput TranslationInput = Connectable.GetInput(UnrealIdentifiers::Translation))
-						{
-							pxr::GfVec2f VecValue;
-							if (TranslationInput.Get(&VecValue))
-							{
-								Translation = FVector2f{VecValue[0], VecValue[1]};
-							}
-						}
+		// Traverse through potentially N UV transform nodes
+		while (pxr::UsdShadeConnectableAPI::GetConnectedSource(StInput.GetAttr(), &Connectable, &ConnectableOutput, &AttributeType))
+		{
+			pxr::UsdPrim ConnectedPrim = Connectable.GetPrim();
+			pxr::UsdShadeNodeDefAPI ConnectedShadeNode{ConnectedPrim};
+			if (!ConnectedShadeNode)
+			{
+				// Deadend, should never really happen
+				return;
+			}
 
-						// Concat transform (scale, then rotation, then translation)
-						FTransform2f NewTransform = FTransform2f{Scale}.Concatenate(
-							FTransform2f{FQuat2f{FMath::DegreesToRadians(Rotation)}, Translation}
-						);
-						ConcatenatedTransform = ConcatenatedTransform.Concatenate(NewTransform);
+			pxr::TfToken ConnectedNodeId;
+			ConnectedShadeNode.GetShaderId(&ConnectedNodeId);
 
-						// Step through to whatever is *this* node's input
-						StInput = Connectable.GetInput(UnrealIdentifiers::In);
-					}
-					// Directly connected to primvar reader
-					else if (ConnectedNodeId == UnrealIdentifiers::UsdPrimvarReader_float2)
+			// UV transform shader node
+			if (ConnectedNodeId == UnrealIdentifiers::UsdTransform2d)
+			{
+				FScale2f Scale{1.0f, 1.0f};
+				if (pxr::UsdShadeInput ScaleInput = Connectable.GetInput(UnrealIdentifiers::Scale))
+				{
+					pxr::GfVec2f VecValue;
+					if (ScaleInput.Get(&VecValue))
 					{
-						bFoundPrimvarReader = true;
-						break;
+						Scale = FScale2f{VecValue[0], VecValue[1]};
 					}
-					else
+				}
+
+				float Rotation = 0.0f;
+				if (pxr::UsdShadeInput RotationInput = Connectable.GetInput(UnrealIdentifiers::Rotation))
+				{
+					float FloatValue;
+					if (RotationInput.Get(&FloatValue))
 					{
-						UE_LOG(
-							LogUsd,
-							Warning,
-							TEXT("Unexpected shader node id '%s' when traversing texture node '%s' for primvars!"),
-							*UsdToUnreal::ConvertToken(ConnectedNodeId),
-							*UsdToUnreal::ConvertPath(UsdUVTexture.GetPrim().GetPrimPath())
-						);
-						return;
+						Rotation = FloatValue;
 					}
 				}
 
-				// Ideally after running through the UsdTransform2d nodes we'd run into a primvar reader node
-				if (bFoundPrimvarReader)
+				FVector2f Translation{0.0f, 0.0f};
+				if (pxr::UsdShadeInput TranslationInput = Connectable.GetInput(UnrealIdentifiers::Translation))
 				{
-					if (pxr::UsdShadeInput VarnameInput = Connectable.GetInput(UnrealIdentifiers::Varname))
+					pxr::GfVec2f VecValue;
+					if (TranslationInput.Get(&VecValue))
 					{
-						// PrimvarReader may have a string literal with the primvar name, or it may be connected to
-						// e.g. an attribute defined elsewhere
-						FString Primvar = RecursivelySearchForStringValue(VarnameInput);
-						if (!Primvar.IsEmpty())
-						{
-							// This stuff will end up as arguments for the UsdTransform2d UE material function, which
-							// will do a bunch of conversions inside. We could precompute some of that here, but instead
-							// we're choosing not to, because this means that these values will show up on the material
-							// instance exactly as they show up in the USD shader prims (e.g. if in USD the rotation is
-							// 30 (degrees) we'll see "30" on the material instance too)
-							OutTextureValue.Primvar = Primvar;
-							OutTextureValue.UVScale = ConcatenatedTransform.GetMatrix().GetScale();
-							OutTextureValue.UVRotation = -FMath::RadiansToDegrees(ConcatenatedTransform.GetMatrix().GetRotationAngle());
-							OutTextureValue.UVTranslation = FVector2f{ConcatenatedTransform.GetTranslation()};
-						}
+						Translation = FVector2f{VecValue[0], VecValue[1]};
 					}
 				}
+
+				// Concat transform (scale, then rotation, then translation)
+				FTransform2f NewTransform = FTransform2f{Scale}.Concatenate(FTransform2f{FQuat2f{FMath::DegreesToRadians(Rotation)}, Translation});
+				ConcatenatedTransform = ConcatenatedTransform.Concatenate(NewTransform);
+
+				// Step through to whatever is *this* node's input
+				StInput = Connectable.GetInput(UnrealIdentifiers::In);
 			}
+			// Directly connected to primvar reader
+			else if (ConnectedNodeId == UnrealIdentifiers::UsdPrimvarReader_float2)
+			{
+				bFoundPrimvarReader = true;
+				break;
+			}
+			else
+			{
+				UE_LOG(
+					LogUsd,
+					Warning,
+					TEXT("Unexpected shader node id '%s' when traversing texture node '%s' for primvars!"),
+					*UsdToUnreal::ConvertToken(ConnectedNodeId),
+					*UsdToUnreal::ConvertPath(UsdUVTexture.GetPrim().GetPrimPath())
+				);
+				return;
+			}
+		}
 
-			bool GetTextureParameterValue(
-				pxr::UsdShadeInput& ShadeInput,
-				TextureGroup Group,
-				FParameterValue& OutValue,
-				const UMaterialInterface* Material,
-				UUsdAssetCache3* TexturesCache,
-				bool bShareAssetsForIdenticalPrims
-			)
+		// Ideally after running through the UsdTransform2d nodes we'd run into a primvar reader node
+		if (bFoundPrimvarReader)
+		{
+			if (pxr::UsdShadeInput VarnameInput = Connectable.GetInput(UnrealIdentifiers::Varname))
 			{
-				FScopedUsdAllocs UsdAllocs;
+				// PrimvarReader may have a string literal with the primvar name, or it may be connected to
+				// e.g. an attribute defined elsewhere
+				FString Primvar = RecursivelySearchForStringValue(VarnameInput);
+				if (!Primvar.IsEmpty())
+				{
+					// This stuff will end up as arguments for the UsdTransform2d UE material function, which
+					// will do a bunch of conversions inside. We could precompute some of that here, but instead
+					// we're choosing not to, because this means that these values will show up on the material
+					// instance exactly as they show up in the USD shader prims (e.g. if in USD the rotation is
+					// 30 (degrees) we'll see "30" on the material instance too)
+					OutTextureValue.Primvar = Primvar;
+					OutTextureValue.UVScale = ConcatenatedTransform.GetMatrix().GetScale();
+					OutTextureValue.UVRotation = -FMath::RadiansToDegrees(ConcatenatedTransform.GetMatrix().GetRotationAngle());
+					OutTextureValue.UVTranslation = FVector2f{ConcatenatedTransform.GetTranslation()};
+				}
+			}
+		}
+	}
+
+	bool GetTextureParameterValue(pxr::UsdShadeInput& ShadeInput, TextureGroup Group, UsdToUnreal::FParameterValue& OutValue)
+	{
+		FScopedUsdAllocs UsdAllocs;
 
-				const bool bIsNormalInput =
-					(ShadeInput.GetTypeName() == pxr::SdfValueTypeNames->Normal3f || ShadeInput.GetTypeName() == pxr::SdfValueTypeNames->Normal3fArray
-					 || ShadeInput.GetBaseName() == UnrealIdentifiers::Normal);
+		const bool bIsNormalInput =
+			(ShadeInput.GetTypeName() == pxr::SdfValueTypeNames->Normal3f || ShadeInput.GetTypeName() == pxr::SdfValueTypeNames->Normal3fArray
+			 || ShadeInput.GetBaseName() == UnrealIdentifiers::Normal);
 
-				pxr::UsdShadeConnectableAPI UsdUVTextureSource;
-				pxr::TfToken UsdUVTextureSourceName;
-				pxr::UsdShadeAttributeType UsdUVTextureAttributeType;
+		pxr::UsdShadeConnectableAPI UsdUVTextureSource;
+		pxr::TfToken UsdUVTextureSourceName;
+		pxr::UsdShadeAttributeType UsdUVTextureAttributeType;
 
-				// Clear it, as we'll signal that it has a valid texture bound by setting it with an FTextureParameterValue
-				OutValue.Set<float>(0.0f);
+		// Clear it, as we'll signal that it has a valid texture bound by setting it with an FTextureParameterValue
+		OutValue.Set<float>(0.0f);
 
+		if (pxr::UsdShadeConnectableAPI::GetConnectedSource(
+				ShadeInput.GetAttr(),
+				&UsdUVTextureSource,
+				&UsdUVTextureSourceName,
+				&UsdUVTextureAttributeType
+			))
+		{
+			pxr::UsdShadeInput FileInput;
+
+			// UsdUVTexture: Get its file input
+			if (UsdUVTextureAttributeType == pxr::UsdShadeAttributeType::Output)
+			{
+				FileInput = UsdUVTextureSource.GetInput(UnrealIdentifiers::File);
+			}
+			// Check if we are being directly passed an asset
+			else
+			{
+				FileInput = UsdUVTextureSource.GetInput(UsdUVTextureSourceName);
+			}
+
+			// Recursively traverse "inputs:file" connections until we stop finding other connected prims
+			pxr::UsdShadeConnectableAPI TextureFileSource;
+			pxr::TfToken TextureFileSourceName;
+			pxr::UsdShadeAttributeType TextureFileAttributeType;
+			while (FileInput)
+			{
 				if (pxr::UsdShadeConnectableAPI::GetConnectedSource(
-						ShadeInput.GetAttr(),
-						&UsdUVTextureSource,
-						&UsdUVTextureSourceName,
-						&UsdUVTextureAttributeType
+						FileInput.GetAttr(),
+						&TextureFileSource,
+						&TextureFileSourceName,
+						&TextureFileAttributeType
 					))
 				{
-					pxr::UsdShadeInput FileInput;
-
-					// UsdUVTexture: Get its file input
-					if (UsdUVTextureAttributeType == pxr::UsdShadeAttributeType::Output)
+					// Another connection, get its file input
+					if (TextureFileAttributeType == pxr::UsdShadeAttributeType::Output)
 					{
-						FileInput = UsdUVTextureSource.GetInput(UnrealIdentifiers::File);
+						FileInput = TextureFileSource.GetInput(UnrealIdentifiers::File);
 					}
 					// Check if we are being directly passed an asset
 					else
 					{
-						FileInput = UsdUVTextureSource.GetInput(UsdUVTextureSourceName);
-					}
-
-					// Recursively traverse "inputs:file" connections until we stop finding other connected prims
-					pxr::UsdShadeConnectableAPI TextureFileSource;
-					pxr::TfToken TextureFileSourceName;
-					pxr::UsdShadeAttributeType TextureFileAttributeType;
-					while (FileInput)
-					{
-						if (pxr::UsdShadeConnectableAPI::GetConnectedSource(
-								FileInput.GetAttr(),
-								&TextureFileSource,
-								&TextureFileSourceName,
-								&TextureFileAttributeType
-							))
-						{
-							// Another connection, get its file input
-							if (TextureFileAttributeType == pxr::UsdShadeAttributeType::Output)
-							{
-								FileInput = TextureFileSource.GetInput(UnrealIdentifiers::File);
-							}
-							// Check if we are being directly passed an asset
-							else
-							{
-								FileInput = TextureFileSource.GetInput(TextureFileSourceName);
-							}
-						}
-						else
-						{
-							break;
-						}
+						FileInput = TextureFileSource.GetInput(TextureFileSourceName);
 					}
+				}
+				else
+				{
+					break;
+				}
+			}
 
-					// Get desired texture wrapping modes
-					TextureAddress AddressX = TextureAddress::TA_Wrap;
-					TextureAddress AddressY = TextureAddress::TA_Wrap;
-					if (pxr::UsdAttribute WrapSAttr = UsdUVTextureSource.GetInput(UnrealIdentifiers::WrapS))
+			if (FileInput && FileInput.GetTypeName() == pxr::SdfValueTypeNames->Asset)	  // Check that FileInput is of type Asset
+			{
+				// Get desired texture wrapping modes
+				TextureAddress AddressX = TextureAddress::TA_Wrap;
+				TextureAddress AddressY = TextureAddress::TA_Wrap;
+				if (pxr::UsdAttribute WrapSAttr = UsdUVTextureSource.GetInput(UnrealIdentifiers::WrapS))
+				{
+					pxr::TfToken WrapS;
+					if (WrapSAttr.Get(&WrapS))
 					{
-						pxr::TfToken WrapS;
-						if (WrapSAttr.Get(&WrapS))
-						{
-							AddressX = WrapS == UnrealIdentifiers::Repeat	? TextureAddress::TA_Wrap
-									   : WrapS == UnrealIdentifiers::Mirror ? TextureAddress::TA_Mirror
-																			: TextureAddress::TA_Clamp;	   // We also consider the "black" wrap mode
-																										   // as clamp as that is the closest we can
-																										   // get
-						}
+						AddressX = WrapS == UnrealIdentifiers::Repeat	? TextureAddress::TA_Wrap
+								   : WrapS == UnrealIdentifiers::Mirror ? TextureAddress::TA_Mirror
+																		: TextureAddress::TA_Clamp;	   // We also consider the "black" wrap
+																									   // mode as clamp as that is the closest
+																									   // we can get
 					}
-					if (pxr::UsdAttribute WrapTAttr = UsdUVTextureSource.GetInput(UnrealIdentifiers::WrapT))
+				}
+				if (pxr::UsdAttribute WrapTAttr = UsdUVTextureSource.GetInput(UnrealIdentifiers::WrapT))
+				{
+					pxr::TfToken WrapT;
+					if (WrapTAttr.Get(&WrapT))
 					{
-						pxr::TfToken WrapT;
-						if (WrapTAttr.Get(&WrapT))
-						{
-							AddressY = WrapT == UnrealIdentifiers::Repeat	? TextureAddress::TA_Wrap
-									   : WrapT == UnrealIdentifiers::Mirror ? TextureAddress::TA_Mirror
-																			: TextureAddress::TA_Clamp;	   // We also consider the "black" wrap mode
-																										   // as clamp as that is the closest we can
-																										   // get
-						}
+						AddressY = WrapT == UnrealIdentifiers::Repeat	? TextureAddress::TA_Wrap
+								   : WrapT == UnrealIdentifiers::Mirror ? TextureAddress::TA_Mirror
+																		: TextureAddress::TA_Clamp;	   // We also consider the "black" wrap
+																									   // mode as clamp as that is the closest
+																									   // we can get
 					}
+				}
 
-					if (FileInput && FileInput.GetTypeName() == pxr::SdfValueTypeNames->Asset)	  // Check that FileInput is of type Asset
-					{
-						const FString TexturePath = UsdUtils::GetResolvedAssetPath(FileInput.GetAttr());
-						// We will assume the texture is valid, and show a warning if we fail to parse this later.
-						// Note that we don't even check that the file exists: If we have a texture bound to opacity then we
-						// assume this material is meant to be translucent, even if the texture path is invalid (or points inside an USDZ archive)
-						if (!TexturePath.IsEmpty())
-						{
-							OutValue.Set<FTextureParameterValue>(FTextureParameterValue{});
-						}
-						else
-						{
-							return false;
-						}
-
-						// We only actually want to retrieve the textures if we have a cache to put them in
-						if (TexturesCache)
-						{
-							// We'll add these to the hash because the materials are built to try and reuse the same textures for multiple channels,
-							// and those may expect linear or sRGB values. Without this we may parse a texture as linear because we hit the opacity
-							// channel first, and then reuse it as linear for the base color channel even though it should have been sRGB. Plus we may
-							// have something weird like a normal map being plugged into the base color and the normal channel.
-							bool bSRGB = true;
-							TextureCompressionSettings CompressionSettings = TextureCompressionSettings::TC_Default;
-							if (bIsNormalInput)
-							{
-								// Disable SRGB when parsing float textures, as they're likely specular/roughness maps
-								bSRGB = false;
-								CompressionSettings = TextureCompressionSettings::TC_Normalmap;
-							}
-							else
-							{
-								if (Group == TEXTUREGROUP_WorldSpecular)
-								{
-									bSRGB = false;
-								}
-
-								// Override the sRGB flag with whatever is specified on the texture shader, if it has anything specific.
-								if (pxr::UsdAttribute SourceColorSpaceAttr = UsdUVTextureSource.GetInput(UnrealIdentifiers::SourceColorSpaceToken))
-								{
-									pxr::TfToken SourceColorSpaceValue;
-									if (SourceColorSpaceAttr.HasAuthoredValue() && SourceColorSpaceAttr.Get(&SourceColorSpaceValue))
-									{
-										bSRGB = SourceColorSpaceValue == UnrealIdentifiers::RawColorSpaceToken	  ? false
-												: SourceColorSpaceValue == UnrealIdentifiers::SRGBColorSpaceToken ? true
-																												  : bSRGB;
-									}
-								}
-							}
-
-							const FString
-								PrefixedTextureHash = UsdUtils::GetAssetHashPrefix(ShadeInput.GetPrim(), bShareAssetsForIdenticalPrims)
-													  + UsdUtils::GetTextureHash(TexturePath, bSRGB, CompressionSettings, AddressX, AddressY);
-
-							pxr::UsdAttribute FileInputAttr = FileInput.GetAttr();
-							const FString ResolvedTexturePath = UsdUtils::GetResolvedAssetPath(FileInputAttr);
-
-							const FString& DesiredTextureName = FPaths::GetBaseFilename(ResolvedTexturePath);
-
-							const EObjectFlags DesiredFlags = RF_Public | RF_Standalone | RF_Transactional;
-
-							bool bCreatedTexture = false;
-							UTexture* Texture = TexturesCache->GetOrCreateCustomCachedAsset<UTexture2D>(
-								PrefixedTextureHash,
-								DesiredTextureName,
-								DesiredFlags,
-								[&ResolvedTexturePath, Group, TexturesCache, bSRGB](UPackage* Outer, FName SanitizedName, EObjectFlags DesiredFlags)
-								{
-									return UsdUtils::CreateTexture(ResolvedTexturePath, SanitizedName, Group, DesiredFlags, Outer, !bSRGB);
-								},
-								&bCreatedTexture
-							);
+				const FString TexturePath = UsdUtils::GetResolvedAssetPath(FileInput.GetAttr());
 
-							if (Texture)
-							{
-								// The texture resource needs to be updated only if the following settings have changed from their default values
-								bool bNeedUpdateResource = false;
-								bNeedUpdateResource |= AddressX != TextureAddress::TA_Wrap;
-								bNeedUpdateResource |= AddressY != TextureAddress::TA_Wrap;
-#if !WITH_EDITOR
-								// In editor, these settings are already set on the factory when the texture is created so no further update is needed
-								bNeedUpdateResource |= bSRGB != true;
-								// TC_Normalmap is set via the TEXTUREGROUP_WorldNormalMap Group
-								bNeedUpdateResource |= CompressionSettings != TextureCompressionSettings::TC_Default;
-#endif	// !WITH_EDITOR
-								if (bCreatedTexture && bNeedUpdateResource)
-								{
-#if WITH_EDITOR
-									Texture->PreEditChange(nullptr);
-#endif	// WITH_EDITOR
-									Texture->SRGB = bSRGB;
-									Texture->CompressionSettings = CompressionSettings;
-									if (UTexture2D* Texture2D = Cast<UTexture2D>(Texture))
-									{
-										Texture2D->AddressX = AddressX;
-										Texture2D->AddressY = AddressY;
-									}
-#if WITH_EDITOR
-									Texture->PostEditChange();
-#else
-									Texture->UpdateResource();
-#endif	// WITH_EDITOR
-								}
-
-								if (UUsdAssetUserData* AssetUserData = UsdUnreal::ObjectUtils::GetOrCreateAssetUserData(Texture))
-								{
-									// Give the same prim path to the texture, so that it ends up imported right next to the material
-									if (UUsdAssetUserData* MaterialAssetUserData = UsdUnreal::ObjectUtils::GetAssetUserData(Material))
-									{
-										for (const FString& MaterialPrimPath : MaterialAssetUserData->PrimPaths)
-										{
-											AssetUserData->PrimPaths.AddUnique(MaterialPrimPath);
-										}
-									}
-								}
-
-								FTextureParameterValue OutTextureValue;
-								OutTextureValue.Texture = Texture;
-								GetSTPrimvarAndTransform(UsdUVTextureSource, OutTextureValue);
-
-								// The UsdUVTexture Shader itself prim has multiple standard outputs, but this is not
-								// full swizzle support (check
-								// [REDACTED_PROFILE_URL]/USD/blob/5c5ebddff35012461a2b0ba773c47f05730cbab4/pxr/usdImaging/plugin/usdShaders/shaders/shaderDefs.usda#L198)
-								if (UsdUVTextureAttributeType == pxr::UsdShadeAttributeType::Output)
-								{
-									const FName OutputName(UsdToUnreal::ConvertToken(UsdUVTextureSourceName));
-
-									if (OutputName == TEXT("rgb"))
-									{
-										OutTextureValue.OutputIndex = 0;
-									}
-									else if (OutputName == TEXT("r"))
-									{
-										OutTextureValue.OutputIndex = 1;
-									}
-									else if (OutputName == TEXT("g"))
-									{
-										OutTextureValue.OutputIndex = 2;
-									}
-									else if (OutputName == TEXT("b"))
-									{
-										OutTextureValue.OutputIndex = 3;
-									}
-									else if (OutputName == TEXT("a"))
-									{
-										OutTextureValue.OutputIndex = 4;
-									}
-								}
-
-								OutValue.Set<FTextureParameterValue>(OutTextureValue);
-							}
-							else
-							{
-								FUsdLogManager::LogMessage(
-									EMessageSeverity::Warning,
-									FText::Format(
-										LOCTEXT("FailedToParseTexture", "Failed to parse texture at path '{0}'"),
-										FText::FromString(TexturePath)
-									)
-								);
-							}
-						}
-					}
+				// We will assume the texture is valid, and show a warning if we fail to parse this later.
+				// Note that we don't even check that the file exists: If we have a texture bound to opacity then we
+				// assume this material is meant to be translucent, even if the texture path is invalid (or points inside an USDZ archive)
+				if (!TexturePath.IsEmpty())
+				{
+					OutValue.Set<UsdToUnreal::FTextureParameterValue>(UsdToUnreal::FTextureParameterValue{});
 				}
-
-				// We may be calling this from IsMaterialTranslucent, when we have no TexturesCache: We wouldn't be able to import textures without
-				// the cache as we may repeatedly parse the same texture. Because of this, and how we clear OutValue to float before we begin, we can
-				// say that if we have a FTextureParameterValue at all, then there is a valid texture that *can* be parsed, and so we return true.
-				return OutValue.IsType<FTextureParameterValue>();
-			}
-
-			bool GetFloatParameterValue(
-				pxr::UsdShadeConnectableAPI& Connectable,
-				const pxr::TfToken& InputName,
-				float DefaultValue,
-				FParameterValue& OutValue,
-				const UMaterialInterface* Material,
-				UUsdAssetCache3* TexturesCache,
-				bool bShareAssetsForIdenticalPrims
-			)
-			{
-				FScopedUsdAllocs Allocs;
-
-				pxr::UsdShadeInput Input = Connectable.GetInput(InputName);
-				if (!Input)
+				else
 				{
 					return false;
 				}
 
-				// If we have another shader/node connected
-				pxr::UsdShadeConnectableAPI Source;
-				pxr::TfToken SourceName;
-				pxr::UsdShadeAttributeType AttributeType;
-				if (pxr::UsdShadeConnectableAPI::GetConnectedSource(Input.GetAttr(), &Source, &SourceName, &AttributeType))
+				// If the texture is UDIM, it will contain the "<UDIM>" tag. UsdUtils::GetResolvedAssetPath will already
+				// replace the tag with one of the tiles so that it plugs nicely into an resolver call that it has, but
+				// (for Interchange) we also need to know whether this path was originally flagged to be interpreted as UDIM
+				// or not, so here we just do a simple check on the string value that is actually on the attribute
+				bool bIsUDIM = false;
+				if (pxr::UsdAttribute Attr = FileInput.GetAttr())
 				{
-					if (!GetTextureParameterValue(Input, TEXTUREGROUP_WorldSpecular, OutValue, Material, TexturesCache, bShareAssetsForIdenticalPrims))
+					pxr::SdfAssetPath Path;
+					if (Attr.Get(&Path))
 					{
-						// Check if we have a fallback input that we can use instead, since we don't have a valid texture value
-						if (const pxr::UsdShadeInput FallbackInput = Source.GetInput(UnrealIdentifiers::Fallback))
-						{
-							float UsdFallbackFloat;
-							if (FallbackInput.Get<float>(&UsdFallbackFloat))
-							{
-								OutValue.Set<float>(UsdFallbackFloat);
-								return true;
-							}
-						}
-
-						// Recurse because the attribute may just be pointing at some other attribute that has the data
-						// (e.g. when shader input is just "hoisted" and connected to the parent material input)
-						return GetFloatParameterValue(Source, SourceName, DefaultValue, OutValue, Material, TexturesCache, bShareAssetsForIdenticalPrims);
+						FString RawFilePath = UsdToUnreal::ConvertString(Path.GetAssetPath());
+						bIsUDIM = RawFilePath.Contains(TEXT("<UDIM>"));
 					}
 				}
-				// No other node connected, so we must have some value
-				else
-				{
-					float InputValue = DefaultValue;
-					Input.Get<float>(&InputValue);
-
-					OutValue.Set<float>(InputValue);
-				}
-
-				return true;
-			}
-
-			bool GetPrimvarReaderParameterValue(
-				const pxr::UsdShadeInput& Input,
-				const pxr::TfToken& PrimvarReaderShaderId,
-				const FLinearColor& DefaultValue,
-				FParameterValue& OutValue
-			)
-			{
-				FScopedUsdAllocs Allocs;
 
-				if (!Input)
+				// Collect desired sRGB, if anything is authored
+				TOptional<bool> bSRGB;
+				if (pxr::UsdAttribute SourceColorSpaceAttr = UsdUVTextureSource.GetInput(UnrealIdentifiers::SourceColorSpaceToken))
 				{
-					return false;
+					pxr::TfToken SourceColorSpaceValue;
+					if (SourceColorSpaceAttr.HasAuthoredValue() && SourceColorSpaceAttr.Get(&SourceColorSpaceValue))
+					{
+						bSRGB = SourceColorSpaceValue == UnrealIdentifiers::RawColorSpaceToken	  ? false
+								: SourceColorSpaceValue == UnrealIdentifiers::SRGBColorSpaceToken ? true
+																								  : bSRGB;
+					}
 				}
 
-				const bool bShaderOutputsOnly = true;
-				const pxr::UsdShadeAttributeVector ValProdAttrs = Input.GetValueProducingAttributes(bShaderOutputsOnly);
-				for (const pxr::UsdAttribute& ValProdAttr : ValProdAttrs)
+				UsdToUnreal::FTextureParameterValue& OutTextureValue = OutValue.Get<UsdToUnreal::FTextureParameterValue>();
+				OutTextureValue.TextureFilePath = TexturePath;
+				OutTextureValue.AddressX = AddressX;
+				OutTextureValue.AddressY = AddressY;
+				OutTextureValue.Group = Group;
+				OutTextureValue.bSRGB = bSRGB;
+				OutTextureValue.bIsUDIM = bIsUDIM;
+				GetSTPrimvarAndTransform(UsdUVTextureSource, OutTextureValue);
+
+				// The UsdUVTexture Shader itself prim has multiple standard outputs, but this is not
+				// full swizzle support (check
+				// [REDACTED_PROFILE_URL]/USD/blob/5c5ebddff35012461a2b0ba773c47f05730cbab4/pxr/usdImaging/plugin/usdShaders/shaders/shaderDefs.usda#L198)
+				if (UsdUVTextureAttributeType == pxr::UsdShadeAttributeType::Output)
 				{
-					const pxr::UsdShadeShader ValProdShader(ValProdAttr.GetPrim());
-					if (!ValProdShader)
-					{
-						continue;
-					}
+					const FName OutputName(UsdToUnreal::ConvertToken(UsdUVTextureSourceName));
 
-					pxr::TfToken ShaderId;
-					if (!ValProdShader.GetShaderId(&ShaderId) || ShaderId != PrimvarReaderShaderId)
+					if (OutputName == TEXT("rgb"))
 					{
-						continue;
+						OutTextureValue.OutputIndex = 0;
 					}
-
-					const pxr::UsdShadeInput VarnameInput = ValProdShader.GetInput(UnrealIdentifiers::Varname);
-					if (!VarnameInput)
+					else if (OutputName == TEXT("r"))
 					{
-						continue;
+						OutTextureValue.OutputIndex = 1;
 					}
-
-					// The schema for UsdPrimvarReader specifies that the "varname" input should be
-					// string-typed, but some assets might be set up token-typed, so we'll consider
-					// either type.
-					std::string PrimvarName;
-					if (VarnameInput.GetTypeName() == pxr::SdfValueTypeNames->String)
+					else if (OutputName == TEXT("g"))
 					{
-						if (!VarnameInput.Get(&PrimvarName))
-						{
-							continue;
-						}
+						OutTextureValue.OutputIndex = 2;
 					}
-					else if (VarnameInput.GetTypeName() == pxr::SdfValueTypeNames->Token)
+					else if (OutputName == TEXT("b"))
 					{
-						pxr::TfToken PrimvarNameToken;
-						if (!VarnameInput.Get(&PrimvarNameToken))
-						{
-							continue;
-						}
-						PrimvarName = PrimvarNameToken.GetString();
+						OutTextureValue.OutputIndex = 3;
 					}
-
-					if (PrimvarName.empty())
+					else if (OutputName == TEXT("a"))
 					{
-						continue;
+						OutTextureValue.OutputIndex = 4;
 					}
+				}
+			}
+		}
 
-					FLinearColor FallbackColor = DefaultValue;
-					const pxr::UsdShadeInput FallbackInput = ValProdShader.GetInput(UnrealIdentifiers::Fallback);
-					pxr::GfVec3f UsdFallbackColor;
-					if (FallbackInput && FallbackInput.Get<pxr::GfVec3f>(&UsdFallbackColor))
-					{
-						FallbackColor = UsdToUnreal::ConvertColor(UsdFallbackColor);
-					}
+		// We may be calling this from IsMaterialTranslucent, when we have no TexturesCache: We wouldn't be able to import textures without
+		// the cache as we may repeatedly parse the same texture. Because of this, and how we clear OutValue to float before we begin, we can
+		// say that if we have a FTextureParameterValue at all, then there is a valid texture that *can* be parsed, and so we return true.
+		return OutValue.IsType<UsdToUnreal::FTextureParameterValue>();
+	}
 
-					OutValue.Set<FPrimvarReaderParameterValue>(FPrimvarReaderParameterValue{PrimvarName, FVector(FallbackColor)});
-					return true;
-				}
+	bool GetFloatParameterValue(
+		pxr::UsdShadeConnectableAPI& Connectable,
+		const pxr::TfToken& InputName,
+		float DefaultValue,
+		UsdToUnreal::FParameterValue& OutValue
+	)
+	{
+		FScopedUsdAllocs Allocs;
 
-				return false;
-			}
+		pxr::UsdShadeInput Input = Connectable.GetInput(InputName);
+		if (!Input)
+		{
+			return false;
+		}
 
-			bool GetVec3ParameterValue(
-				pxr::UsdShadeConnectableAPI& Connectable,
-				const pxr::TfToken& InputName,
-				const FLinearColor& DefaultValue,
-				FParameterValue& OutValue,
-				bool bIsNormalMap,
-				const UMaterialInterface* Material,
-				UUsdAssetCache3* TexturesCache,
-				bool bShareAssetsForIdenticalPrims
-			)
+		// If we have another shader/node connected
+		pxr::UsdShadeConnectableAPI Source;
+		pxr::TfToken SourceName;
+		pxr::UsdShadeAttributeType AttributeType;
+		if (pxr::UsdShadeConnectableAPI::GetConnectedSource(Input.GetAttr(), &Source, &SourceName, &AttributeType))
+		{
+			if (!GetTextureParameterValue(Input, TEXTUREGROUP_WorldSpecular, OutValue))
 			{
-				FScopedUsdAllocs Allocs;
-
-				pxr::UsdShadeInput Input = Connectable.GetInput(InputName);
-				if (!Input)
+				// Check if we have a fallback input that we can use instead, since we don't have a valid texture value
+				if (const pxr::UsdShadeInput FallbackInput = Source.GetInput(UnrealIdentifiers::Fallback))
 				{
-					return false;
+					float UsdFallbackFloat;
+					if (FallbackInput.Get<float>(&UsdFallbackFloat))
+					{
+						OutValue.Set<float>(UsdFallbackFloat);
+						return true;
+					}
 				}
 
-				// If we have another shader/node connected
-				pxr::UsdShadeConnectableAPI Source;
-				pxr::TfToken SourceName;
-				pxr::UsdShadeAttributeType AttributeType;
-				if (pxr::UsdShadeConnectableAPI::GetConnectedSource(Input.GetAttr(), &Source, &SourceName, &AttributeType))
-				{
-					if (!GetTextureParameterValue(
-							Input,
-							bIsNormalMap ? TEXTUREGROUP_WorldNormalMap : TEXTUREGROUP_World,
-							OutValue,
-							Material,
-							TexturesCache,
-							bShareAssetsForIdenticalPrims
-						))
-					{
-						// Check whether this input receives its value through a connection to a
-						// primvar reader shader.
-						if (GetPrimvarReaderParameterValue(Input, UnrealIdentifiers::UsdPrimvarReader_float3, DefaultValue, OutValue))
-						{
-							return true;
-						}
+				// Recurse because the attribute may just be pointing at some other attribute that has the data
+				// (e.g. when shader input is just "hoisted" and connected to the parent material input)
+				return GetFloatParameterValue(Source, SourceName, DefaultValue, OutValue);
+			}
+		}
+		// No other node connected, so we must have some value
+		else
+		{
+			float InputValue = DefaultValue;
+			Input.Get<float>(&InputValue);
 
-						// Check if we have a fallback input that we can use instead, since we don't have a valid texture value
-						if (const pxr::UsdShadeInput FallbackInput = Source.GetInput(UnrealIdentifiers::Fallback))
-						{
-							pxr::GfVec3f UsdFallbackVec3;
-							if (FallbackInput.Get<pxr::GfVec3f>(&UsdFallbackVec3))
-							{
-								OutValue.Set<FVector>(UsdToUnreal::ConvertVector(UsdFallbackVec3));
-								return true;
-							}
-
-							pxr::GfVec4f UsdFallbackVec4;
-							if (FallbackInput.Get<pxr::GfVec4f>(&UsdFallbackVec4))
-							{
-								if (!FMath::IsNearlyEqual(UsdFallbackVec4[3], 1.0f))
-								{
-									UE_LOG(
-										LogUsd,
-										Warning,
-										TEXT("Ignoring alpha value from fallback GfVec4f [%f, %f, %f, %f] used for Shader '%s'"),
-										UsdFallbackVec4[0],
-										UsdFallbackVec4[1],
-										UsdFallbackVec4[2],
-										UsdFallbackVec4[3],
-										*UsdToUnreal::ConvertPath(Source.GetPrim().GetPrimPath())
-									);
-								}
-
-								OutValue.Set<FVector>(FVector{UsdFallbackVec4[0], UsdFallbackVec4[1], UsdFallbackVec4[2]});
-								return true;
-							}
-						}
+			OutValue.Set<float>(InputValue);
+		}
 
-						// This shader doesn't have anything: Traverse into the input connectable itself
-						return GetVec3ParameterValue(
-							Source,
-							SourceName,
-							DefaultValue,
-							OutValue,
-							bIsNormalMap,
-							Material,
-							TexturesCache,
-							bShareAssetsForIdenticalPrims
-						);
-					}
-				}
-				// No other node connected, so we must have some value
-				else if (InputName != UnrealIdentifiers::Normal)
-				{
-					FLinearColor DiffuseColor = DefaultValue;
-					pxr::GfVec3f UsdDiffuseColor;
-					if (Input.Get<pxr::GfVec3f>(&UsdDiffuseColor))
-					{
-						DiffuseColor = UsdToUnreal::ConvertColor(UsdDiffuseColor);
-					}
+		return true;
+	}
 
-					OutValue.Set<FVector>(FVector(DiffuseColor));
-				}
+	bool GetPrimvarReaderParameterValue(
+		const pxr::UsdShadeInput& Input,
+		const pxr::TfToken& PrimvarReaderShaderId,
+		const FLinearColor& DefaultValue,
+		UsdToUnreal::FParameterValue& OutValue
+	)
+	{
+		FScopedUsdAllocs Allocs;
 
-				return true;
+		if (!Input)
+		{
+			return false;
+		}
+
+		const bool bShaderOutputsOnly = true;
+		const pxr::UsdShadeAttributeVector ValProdAttrs = Input.GetValueProducingAttributes(bShaderOutputsOnly);
+		for (const pxr::UsdAttribute& ValProdAttr : ValProdAttrs)
+		{
+			const pxr::UsdShadeShader ValProdShader(ValProdAttr.GetPrim());
+			if (!ValProdShader)
+			{
+				continue;
 			}
 
-			bool GetBoolParameterValue(
-				pxr::UsdShadeConnectableAPI& Connectable,
-				const pxr::TfToken& InputName,
-				bool DefaultValue,
-				FParameterValue& OutValue,
-				UMaterialInterface* Material = nullptr,
-				UUsdAssetCache3* TexturesCache = nullptr
-			)
+			pxr::TfToken ShaderId;
+			if (!ValProdShader.GetShaderId(&ShaderId) || ShaderId != PrimvarReaderShaderId)
 			{
-				FScopedUsdAllocs Allocs;
+				continue;
+			}
 
-				pxr::UsdShadeInput Input = Connectable.GetInput(InputName);
-				if (!Input)
-				{
-					return false;
-				}
+			const pxr::UsdShadeInput VarnameInput = ValProdShader.GetInput(UnrealIdentifiers::Varname);
+			if (!VarnameInput)
+			{
+				continue;
+			}
 
-				// If we have another shader/node connected
-				pxr::UsdShadeConnectableAPI Source;
-				pxr::TfToken SourceName;
-				pxr::UsdShadeAttributeType AttributeType;
-				if (pxr::UsdShadeConnectableAPI::GetConnectedSource(Input.GetAttr(), &Source, &SourceName, &AttributeType))
+			// The schema for UsdPrimvarReader specifies that the "varname" input should be
+			// string-typed, but some assets might be set up token-typed, so we'll consider
+			// either type.
+			std::string PrimvarName;
+			if (VarnameInput.GetTypeName() == pxr::SdfValueTypeNames->String)
+			{
+				if (!VarnameInput.Get(&PrimvarName))
 				{
-					// Recurse because the attribute may just be pointing at some other attribute that has the data
-					// (e.g. when shader input is just "hoisted" and connected to the parent material input)
-					return GetBoolParameterValue(Source, SourceName, DefaultValue, OutValue, Material, TexturesCache);
+					continue;
 				}
-				// No other node connected, so we must have some value
-				else
+			}
+			else if (VarnameInput.GetTypeName() == pxr::SdfValueTypeNames->Token)
+			{
+				pxr::TfToken PrimvarNameToken;
+				if (!VarnameInput.Get(&PrimvarNameToken))
 				{
-					bool InputValue = DefaultValue;
-					Input.Get<bool>(&InputValue);
-
-					OutValue.Set<bool>(InputValue);
+					continue;
 				}
+				PrimvarName = PrimvarNameToken.GetString();
+			}
 
-				return true;
+			if (PrimvarName.empty())
+			{
+				continue;
 			}
 
-#if WITH_EDITOR
-			/**
-			 * Intended to be used with Visit(), this will create a UMaterialExpression in InMaterial, set it with the
-			 * value stored in the current variant of an FParameterValue, and return it.
-			 */
-			struct FGetExpressionForValueVisitor
-			{
-				FGetExpressionForValueVisitor(UMaterial& InMaterial)
-					: Material(InMaterial)
-				{
-				}
+			FLinearColor FallbackColor = DefaultValue;
+			const pxr::UsdShadeInput FallbackInput = ValProdShader.GetInput(UnrealIdentifiers::Fallback);
+			pxr::GfVec3f UsdFallbackColor;
+			if (FallbackInput && FallbackInput.Get<pxr::GfVec3f>(&UsdFallbackColor))
+			{
+				FallbackColor = UsdToUnreal::ConvertColor(UsdFallbackColor);
+			}
 
-				UMaterialExpression* operator()(const float FloatValue) const
-				{
-					UMaterialExpressionConstant* Expression = Cast<UMaterialExpressionConstant>(
-						UMaterialEditingLibrary::CreateMaterialExpression(&Material, UMaterialExpressionConstant::StaticClass())
-					);
-					Expression->R = FloatValue;
+			OutValue.Set<UsdToUnreal::FPrimvarReaderParameterValue>(
+				UsdToUnreal::FPrimvarReaderParameterValue{UsdToUnreal::ConvertString(PrimvarName), FVector(FallbackColor)}
+			);
+			return true;
+		}
 
-					return Expression;
-				}
+		return false;
+	}
 
-				UMaterialExpression* operator()(const FVector& VectorValue) const
-				{
-					UMaterialExpressionConstant4Vector* Expression = Cast<UMaterialExpressionConstant4Vector>(
-						UMaterialEditingLibrary::CreateMaterialExpression(&Material, UMaterialExpressionConstant4Vector::StaticClass())
-					);
-					Expression->Constant = FLinearColor(VectorValue);
+	bool GetVec3ParameterValue(
+		pxr::UsdShadeConnectableAPI& Connectable,
+		const pxr::TfToken& InputName,
+		const FLinearColor& DefaultValue,
+		UsdToUnreal::FParameterValue& OutValue,
+		TextureGroup Group = TextureGroup::TEXTUREGROUP_World
+	)
+	{
+		FScopedUsdAllocs Allocs;
 
-					return Expression;
-				}
+		pxr::UsdShadeInput Input = Connectable.GetInput(InputName);
+		if (!Input)
+		{
+			return false;
+		}
 
-				UMaterialExpression* operator()(const FTextureParameterValue& TextureValue) const
+		// If we have another shader/node connected
+		pxr::UsdShadeConnectableAPI Source;
+		pxr::TfToken SourceName;
+		pxr::UsdShadeAttributeType AttributeType;
+		if (pxr::UsdShadeConnectableAPI::GetConnectedSource(Input.GetAttr(), &Source, &SourceName, &AttributeType))
+		{
+			if (!GetTextureParameterValue(Input, Group, OutValue))
+			{
+				// Check whether this input receives its value through a connection to a
+				// primvar reader shader.
+				if (GetPrimvarReaderParameterValue(Input, UnrealIdentifiers::UsdPrimvarReader_float3, DefaultValue, OutValue))
 				{
-					UMaterialExpressionTextureSample* Expression = Cast<UMaterialExpressionTextureSample>(
-						UMaterialEditingLibrary::CreateMaterialExpression(&Material, UMaterialExpressionTextureSample::StaticClass())
-					);
-					Expression->Texture = TextureValue.Texture;
-					Expression->SamplerType = UMaterialExpressionTextureBase::GetSamplerTypeForTexture(TextureValue.Texture);
-
-					return Expression;
+					return true;
 				}
 
-				UMaterialExpression* operator()(const FPrimvarReaderParameterValue& PrimvarReaderValue) const
+				// Check if we have a fallback input that we can use instead, since we don't have a valid texture value
+				if (const pxr::UsdShadeInput FallbackInput = Source.GetInput(UnrealIdentifiers::Fallback))
 				{
-					UMaterialExpression* Expression = nullptr;
-
-					// We currently only import the "displayColor" primvar in USD as vertex colors
-					// on UE meshes, so that's the only primvar we can support here.
-					if (PrimvarReaderValue.PrimvarName == "displayColor")
+					pxr::GfVec3f UsdFallbackVec3;
+					if (FallbackInput.Get<pxr::GfVec3f>(&UsdFallbackVec3))
 					{
-						Expression = UMaterialEditingLibrary::CreateMaterialExpression(&Material, UMaterialExpressionVertexColor::StaticClass());
+						OutValue.Set<FVector>(UsdToUnreal::ConvertVector(UsdFallbackVec3));
+						return true;
 					}
-					else
-					{
-						Expression = UMaterialEditingLibrary::CreateMaterialExpression(&Material, UMaterialExpressionConstant4Vector::StaticClass());
 
-						UMaterialExpressionConstant4Vector* FallbackExpression = Cast<UMaterialExpressionConstant4Vector>(Expression);
-						if (FallbackExpression)
+					pxr::GfVec4f UsdFallbackVec4;
+					if (FallbackInput.Get<pxr::GfVec4f>(&UsdFallbackVec4))
+					{
+						if (!FMath::IsNearlyEqual(UsdFallbackVec4[3], 1.0f))
 						{
-							FallbackExpression->Constant = FLinearColor(PrimvarReaderValue.FallbackValue);
+							UE_LOG(
+								LogUsd,
+								Warning,
+								TEXT("Ignoring alpha value from fallback GfVec4f [%f, %f, %f, %f] used for Shader '%s'"),
+								UsdFallbackVec4[0],
+								UsdFallbackVec4[1],
+								UsdFallbackVec4[2],
+								UsdFallbackVec4[3],
+								*UsdToUnreal::ConvertPath(Source.GetPrim().GetPrimPath())
+							);
 						}
-					}
 
-					return Expression;
+						OutValue.Set<FVector>(FVector{UsdFallbackVec4[0], UsdFallbackVec4[1], UsdFallbackVec4[2]});
+						return true;
+					}
 				}
 
-			private:
-				UMaterial& Material;
-			};
-#endif	  // WITH_EDITOR
-
-			UMaterialExpression* GetExpressionForValue(UMaterial& Material, const FParameterValue& ParameterValue)
-			{
-#if WITH_EDITOR
-				FGetExpressionForValueVisitor Visitor{Material};
-				return Visit(Visitor, ParameterValue);
-#else
-				return nullptr;
-#endif	  // WITH_EDITOR
+				// This shader doesn't have anything: Traverse into the input connectable itself
+				return GetVec3ParameterValue(Source, SourceName, DefaultValue, OutValue, Group);
 			}
-
-			/**
-			 * Intended to be used with Visit(), this will set an FParameterValue into InMaterial using InParameterName,
-			 * depending on the variant active in FParameterValue
-			 */
-			struct FSetParameterValueVisitor
+		}
+		// No other node connected, so we must have some value
+		else if (InputName != UnrealIdentifiers::Normal)
+		{
+			FLinearColor DiffuseColor = DefaultValue;
+			pxr::GfVec3f UsdDiffuseColor;
+			if (Input.Get<pxr::GfVec3f>(&UsdDiffuseColor))
 			{
-				explicit FSetParameterValueVisitor(UMaterialInstance& InMaterial, const TCHAR* InParameterName)
-					: Material(InMaterial)
-					, ParameterName(InParameterName)
-				{
-				}
-
-				void operator()(const float FloatValue) const
-				{
-					UsdUtils::SetScalarParameterValue(Material, ParameterName, FloatValue);
-				}
+				DiffuseColor = UsdToUnreal::ConvertColor(UsdDiffuseColor);
+			}
 
-				void operator()(const FVector& VectorValue) const
-				{
-					UsdUtils::SetVectorParameterValue(Material, ParameterName, FLinearColor(VectorValue));
-				}
+			OutValue.Set<FVector>(FVector(DiffuseColor));
+		}
 
-				void operator()(const FTextureParameterValue& TextureValue) const
-				{
-					UsdUtils::SetTextureParameterValue(Material, ParameterName, TextureValue.Texture);
-				}
+		return true;
+	}
 
-				void operator()(const FPrimvarReaderParameterValue& PrimvarReaderValue) const
-				{
-					UsdUtils::SetVectorParameterValue(Material, ParameterName, FLinearColor(PrimvarReaderValue.FallbackValue));
-				}
+	bool GetBoolParameterValue(
+		pxr::UsdShadeConnectableAPI& Connectable,
+		const pxr::TfToken& InputName,
+		bool DefaultValue,
+		UsdToUnreal::FParameterValue& OutValue
+	)
+	{
+		FScopedUsdAllocs Allocs;
 
-				void operator()(const bool BoolValue) const
-				{
-					UsdUtils::SetBoolParameterValue(Material, ParameterName, BoolValue);
-				}
+		pxr::UsdShadeInput Input = Connectable.GetInput(InputName);
+		if (!Input)
+		{
+			return false;
+		}
 
-			protected:
-				UMaterialInstance& Material;
-				const TCHAR* ParameterName;
-			};
-
-			/**
-			 * Specialized version of FSetParameterValueVisitor for UE's UsdPreviewSurface reference materials.
-			 */
-			struct FSetPreviewSurfaceParameterValueVisitor : private FSetParameterValueVisitor
-			{
-				explicit FSetPreviewSurfaceParameterValueVisitor(
-					UMaterialInstance& InMaterial,
-					const TCHAR* InParameterName,
-					const TMap<FString, int32>& InPrimvarToUVIndex
-				)
-					: FSetParameterValueVisitor(InMaterial, InParameterName)
-					, PrimvarToUVIndex(InPrimvarToUVIndex)
-				{
-				}
+		// If we have another shader/node connected
+		pxr::UsdShadeConnectableAPI Source;
+		pxr::TfToken SourceName;
+		pxr::UsdShadeAttributeType AttributeType;
+		if (pxr::UsdShadeConnectableAPI::GetConnectedSource(Input.GetAttr(), &Source, &SourceName, &AttributeType))
+		{
+			// Recurse because the attribute may just be pointing at some other attribute that has the data
+			// (e.g. when shader input is just "hoisted" and connected to the parent material input)
+			return GetBoolParameterValue(Source, SourceName, DefaultValue, OutValue);
+		}
+		// No other node connected, so we must have some value
+		else
+		{
+			bool InputValue = DefaultValue;
+			Input.Get<bool>(&InputValue);
 
-				void operator()(const float FloatValue) const
-				{
-					FSetParameterValueVisitor::operator()(FloatValue);
-					UsdUtils::SetScalarParameterValue(Material, *FString::Printf(TEXT("Use%sTexture"), ParameterName), 0.0f);
-				}
+			OutValue.Set<bool>(InputValue);
+		}
 
-				void operator()(const FVector& VectorValue) const
-				{
-					FSetParameterValueVisitor::operator()(VectorValue);
-					UsdUtils::SetScalarParameterValue(Material, *FString::Printf(TEXT("Use%sTexture"), ParameterName), 0.0f);
-				}
+		return true;
+	}
 
-				void operator()(const FTextureParameterValue& TextureValue) const
-				{
-					UsdUtils::SetTextureParameterValue(Material, *FString::Printf(TEXT("%sTexture"), ParameterName), TextureValue.Texture);
-					UsdUtils::SetScalarParameterValue(Material, *FString::Printf(TEXT("Use%sTexture"), ParameterName), 1.0f);
+#if WITH_EDITOR
+	/**
+	 * Intended to be used with Visit(), this will create a UMaterialExpression in InMaterial, set it with the
+	 * value stored in the current variant of an FParameterValue, and return it.
+	 */
+	struct FGetExpressionForValueVisitor
+	{
+		FGetExpressionForValueVisitor(UMaterial& InMaterial)
+			: Material(InMaterial)
+		{
+		}
 
-					FLinearColor ScaleAndTranslation = FLinearColor{
-						TextureValue.UVScale.GetVector()[0],
-						TextureValue.UVScale.GetVector()[1],
-						TextureValue.UVTranslation[0],
-						TextureValue.UVTranslation[1]};
-					UsdUtils::SetVectorParameterValue(Material, *FString::Printf(TEXT("%sScaleTranslation"), ParameterName), ScaleAndTranslation);
+		UMaterialExpression* operator()(const float FloatValue) const
+		{
+			UMaterialExpressionConstant* Expression = Cast<UMaterialExpressionConstant>(
+				UMaterialEditingLibrary::CreateMaterialExpression(&Material, UMaterialExpressionConstant::StaticClass())
+			);
+			Expression->R = FloatValue;
 
-					UsdUtils::SetScalarParameterValue(Material, *FString::Printf(TEXT("%sRotation"), ParameterName), TextureValue.UVRotation);
+			return Expression;
+		}
 
-					if (const int32* FoundIndex = PrimvarToUVIndex.Find(TextureValue.Primvar))
-					{
-						UsdUtils::SetScalarParameterValue(Material, *FString::Printf(TEXT("%sUVIndex"), ParameterName), *FoundIndex);
-					}
-					else
-					{
-						UE_LOG(
-							LogUsd,
-							Warning,
-							TEXT("Failed to find primvar '%s' when setting material parameter '%s' on material '%s'. Available primvars and UV "
-								 "indices: %s.%s"),
-							*TextureValue.Primvar,
-							ParameterName,
-							*Material.GetPathName(),
-							*UsdUtils::StringifyMap(PrimvarToUVIndex),
-							TextureValue.Primvar.IsEmpty() ? TEXT(" Is your UsdUVTexture Shader missing the 'inputs:st' attribute? (It specifies "
-																  "which UV set to sample the texture with)")
-														   : TEXT("")
-						);
-					}
+		UMaterialExpression* operator()(const FVector& VectorValue) const
+		{
+			UMaterialExpressionConstant4Vector* Expression = Cast<UMaterialExpressionConstant4Vector>(
+				UMaterialEditingLibrary::CreateMaterialExpression(&Material, UMaterialExpressionConstant4Vector::StaticClass())
+			);
+			Expression->Constant = FLinearColor(VectorValue);
 
-					FLinearColor ComponentMask = FLinearColor::Black;
-					switch (TextureValue.OutputIndex)
-					{
-						case 0:	   // RGB
-							ComponentMask = FLinearColor{1.f, 1.f, 1.f, 0.f};
-							break;
-						case 1:	   // R
-							ComponentMask = FLinearColor{1.f, 0.f, 0.f, 0.f};
-							break;
-						case 2:	   // G
-							ComponentMask = FLinearColor{0.f, 1.f, 0.f, 0.f};
-							break;
-						case 3:	   // B
-							ComponentMask = FLinearColor{0.f, 0.f, 1.f, 0.f};
-							break;
-						case 4:	   // A
-							ComponentMask = FLinearColor{0.f, 0.f, 0.f, 1.f};
-							break;
-					}
-					UsdUtils::SetVectorParameterValue(Material, *FString::Printf(TEXT("%sTextureComponent"), ParameterName), ComponentMask);
-				}
+			return Expression;
+		}
 
-				void operator()(const FPrimvarReaderParameterValue& PrimvarReaderValue) const
-				{
-					FSetParameterValueVisitor::operator()(PrimvarReaderValue.FallbackValue);
+		UMaterialExpression* operator()(const UsdToUnreal::FTextureParameterValue& TextureValue) const
+		{
+			UMaterialExpressionTextureSample* Expression = Cast<UMaterialExpressionTextureSample>(
+				UMaterialEditingLibrary::CreateMaterialExpression(&Material, UMaterialExpressionTextureSample::StaticClass())
+			);
+			Expression->Texture = TextureValue.Texture;
+			Expression->SamplerType = UMaterialExpressionTextureBase::GetSamplerTypeForTexture(TextureValue.Texture);
 
-					// We currently only import the "displayColor" primvar in USD as vertex colors
-					// on UE meshes, so that's the only primvar we can support here.
-					if (PrimvarReaderValue.PrimvarName == "displayColor")
-					{
-						UsdUtils::SetScalarParameterValue(Material, TEXT("UseVertexColorForBaseColor"), 1.0f);
-					}
-				}
+			return Expression;
+		}
 
-			protected:
-				const TMap<FString, int32>& PrimvarToUVIndex;
-			};
+		UMaterialExpression* operator()(const UsdToUnreal::FPrimvarReaderParameterValue& PrimvarReaderValue) const
+		{
+			UMaterialExpression* Expression = nullptr;
 
-			void SetParameterValue(
-				UMaterialInstance& Material,
-				const TCHAR* ParameterName,
-				const FParameterValue& ParameterValue,
-				bool bForUsdPreviewSurface,
-				const TMap<FString, int32>& PrimvarToUVIndex
-			)
+			// We currently only import the "displayColor" primvar in USD as vertex colors
+			// on UE meshes, so that's the only primvar we can support here.
+			if (PrimvarReaderValue.PrimvarName == "displayColor")
 			{
-				if (bForUsdPreviewSurface)
-				{
-					FSetPreviewSurfaceParameterValueVisitor Visitor{Material, ParameterName, PrimvarToUVIndex};
-					Visit(Visitor, ParameterValue);
-				}
-				else
-				{
-					FSetParameterValueVisitor Visitor{Material, ParameterName};
-					Visit(Visitor, ParameterValue);
-				}
+				Expression = UMaterialEditingLibrary::CreateMaterialExpression(&Material, UMaterialExpressionVertexColor::StaticClass());
 			}
-
-			UTexture* CreateTextureWithEditor(
-				const FString& ResolvedTexturePath,
-				FName SanitizedName,
-				TextureGroup Group,
-				EObjectFlags ObjectFlags,
-				UObject* Outer,
-				bool bForceLinear
-			)
+			else
 			{
-				UTexture* Texture = nullptr;
-#if WITH_EDITOR
-				if (ResolvedTexturePath.IsEmpty() || !Outer)
+				Expression = UMaterialEditingLibrary::CreateMaterialExpression(&Material, UMaterialExpressionConstant4Vector::StaticClass());
+
+				UMaterialExpressionConstant4Vector* FallbackExpression = Cast<UMaterialExpressionConstant4Vector>(Expression);
+				if (FallbackExpression)
 				{
-					return nullptr;
+					FallbackExpression->Constant = FLinearColor(PrimvarReaderValue.FallbackValue);
 				}
+			}
 
-				FScopedUnrealAllocs UnrealAllocs;
+			return Expression;
+		}
 
-				UTextureFactory* TextureFactory = NewObject<UTextureFactory>();
-				TextureFactory->SuppressImportOverwriteDialog();
-				TextureFactory->bUseHashAsGuid = true;
-				TextureFactory->LODGroup = Group;
-				TextureFactory->HDRImportShouldBeLongLatCubeMap = EAppReturnType::YesAll;
-				// To maintain existing behavior, ColorSpaceMode is left at auto if linear is not requested
-				TextureFactory->ColorSpaceMode = bForceLinear ? ETextureSourceColorSpace::Linear : ETextureSourceColorSpace::Auto;
+	private:
+		UMaterial& Material;
+	};
+#endif	  // WITH_EDITOR
 
-				const bool bIsSupportedUdimTexture = ResolvedTexturePath.Contains(TEXT("<UDIM>"));
-				if (bIsSupportedUdimTexture)
-				{
-					FString BaseFileName = FPaths::GetBaseFilename(ResolvedTexturePath);
+	UMaterialExpression* GetExpressionForValue(UMaterial& Material, const UsdToUnreal::FParameterValue& ParameterValue)
+	{
+#if WITH_EDITOR
+		FGetExpressionForValueVisitor Visitor{Material};
+		return Visit(Visitor, ParameterValue);
+#else
+		return nullptr;
+#endif	  // WITH_EDITOR
+	}
 
-					FString BaseFileNameBeforeUdim;
-					FString BaseFileNameAfterUdim;
-					BaseFileName.Split(TEXT("<UDIM>"), &BaseFileNameBeforeUdim, &BaseFileNameAfterUdim);
+	/**
+	 * Intended to be used with Visit(), this will set an FParameterValue into InMaterial using InParameterName,
+	 * depending on the variant active in FParameterValue
+	 */
+	struct FSetParameterValueVisitor
+	{
+		explicit FSetParameterValueVisitor(UMaterialInstance& InMaterial, const TCHAR* InParameterName)
+			: Material(InMaterial)
+			, ParameterName(InParameterName)
+		{
+		}
 
-					FString UdimRegexPattern = FString::Printf(TEXT(R"((%s)(\d{4})(%s))"), *BaseFileNameBeforeUdim, *BaseFileNameAfterUdim);
-					TextureFactory->UdimRegexPattern = MoveTemp(UdimRegexPattern);
-				}
+		void operator()(const float FloatValue) const
+		{
+			UsdUtils::SetScalarParameterValue(Material, ParameterName, FloatValue);
+		}
 
-				if (!ResolvedTexturePath.IsEmpty())
-				{
-					FString TextureExtension;
-					if (IsInsideUsdzArchive(ResolvedTexturePath, TextureExtension))
-					{
-						// Always prefer using the TextureFactory if we can, as it may provide compression, which the runtime version never will
-						Texture = ReadTextureFromUsdzArchiveEditor(ResolvedTexturePath, TextureExtension, TextureFactory, Outer, ObjectFlags);
-					}
-					// Not inside an USDZ archive, just a regular texture
-					else
-					{
-						bool bOutCancelled = false;
-						const TCHAR* Params = TEXT("");
-						Texture = Cast<UTexture>(TextureFactory->ImportObject(	  //
-							UTexture::StaticClass(),
-							Outer,
-							SanitizedName,
-							ObjectFlags,
-							ResolvedTexturePath,
-							Params,
-							bOutCancelled
-						));
-					}
+		void operator()(const FVector& VectorValue) const
+		{
+			UsdUtils::SetVectorParameterValue(Material, ParameterName, FLinearColor(VectorValue));
+		}
 
-					if (Texture)
-					{
-						// We set this even if we're not going to import so that we can track our original texture filepath
-						// in case we later do an Actions->Import
-						UUsdAssetImportData* ImportData = NewObject<UUsdAssetImportData>(Texture);
-						ImportData->UpdateFilenameOnly(ResolvedTexturePath);
-						Texture->AssetImportData = ImportData;
-					}
-				}
-#endif	  // WITH_EDITOR
-				return Texture;
-			}
+		void operator()(const UsdToUnreal::FTextureParameterValue& TextureValue) const
+		{
+			UsdUtils::SetTextureParameterValue(Material, ParameterName, TextureValue.Texture);
+		}
 
-			UTexture* CreateTextureAtRuntime(const FString& ResolvedTexturePath)
-			{
-				FScopedUnrealAllocs UnrealAllocs;
+		void operator()(const UsdToUnreal::FPrimvarReaderParameterValue& PrimvarReaderValue) const
+		{
+			UsdUtils::SetVectorParameterValue(Material, ParameterName, FLinearColor(PrimvarReaderValue.FallbackValue));
+		}
 
-				UTexture* Texture = nullptr;
+		void operator()(const bool BoolValue) const
+		{
+			UsdUtils::SetBoolParameterValue(Material, ParameterName, BoolValue);
+		}
 
-				if (!ResolvedTexturePath.IsEmpty())
-				{
-					// Try checking if the texture is inside an USDZ archive first, or else TextureFactory throws an error
-					FString TextureExtension;
-					if (IsInsideUsdzArchive(ResolvedTexturePath, TextureExtension))
-					{
-						Texture = ReadTextureFromUsdzArchiveRuntime(ResolvedTexturePath);
-					}
+	protected:
+		UMaterialInstance& Material;
+		const TCHAR* ParameterName;
+	};
 
-					// Not inside an USDZ archive, just a regular texture
-					if (!Texture)
-					{
-						Texture = FImageUtils::ImportFileAsTexture2D(ResolvedTexturePath);
-					}
-				}
+	/**
+	 * Specialized version of FSetParameterValueVisitor for UE's UsdPreviewSurface reference materials.
+	 */
+	struct FSetPreviewSurfaceParameterValueVisitor : private FSetParameterValueVisitor
+	{
+		explicit FSetPreviewSurfaceParameterValueVisitor(
+			UMaterialInstance& InMaterial,
+			const TCHAR* InParameterName,
+			const TMap<FString, int32>& InPrimvarToUVIndex
+		)
+			: FSetParameterValueVisitor(InMaterial, InParameterName)
+			, PrimvarToUVIndex(InPrimvarToUVIndex)
+		{
+		}
+
+		void operator()(const float FloatValue) const
+		{
+			FSetParameterValueVisitor::operator()(FloatValue);
+			UsdUtils::SetScalarParameterValue(Material, *FString::Printf(TEXT("Use%sTexture"), ParameterName), 0.0f);
+		}
+
+		void operator()(const FVector& VectorValue) const
+		{
+			FSetParameterValueVisitor::operator()(VectorValue);
+			UsdUtils::SetScalarParameterValue(Material, *FString::Printf(TEXT("Use%sTexture"), ParameterName), 0.0f);
+		}
 
-				return Texture;
+		void operator()(const UsdToUnreal::FTextureParameterValue& TextureValue) const
+		{
+			UsdUtils::SetTextureParameterValue(Material, *FString::Printf(TEXT("%sTexture"), ParameterName), TextureValue.Texture);
+			UsdUtils::SetScalarParameterValue(Material, *FString::Printf(TEXT("Use%sTexture"), ParameterName), 1.0f);
+
+			FLinearColor ScaleAndTranslation = FLinearColor{
+				TextureValue.UVScale.GetVector()[0],
+				TextureValue.UVScale.GetVector()[1],
+				TextureValue.UVTranslation[0],
+				TextureValue.UVTranslation[1]};
+			UsdUtils::SetVectorParameterValue(Material, *FString::Printf(TEXT("%sScaleTranslation"), ParameterName), ScaleAndTranslation);
+
+			UsdUtils::SetScalarParameterValue(Material, *FString::Printf(TEXT("%sRotation"), ParameterName), TextureValue.UVRotation);
+
+			if (const int32* FoundIndex = PrimvarToUVIndex.Find(TextureValue.Primvar))
+			{
+				UsdUtils::SetScalarParameterValue(Material, *FString::Printf(TEXT("%sUVIndex"), ParameterName), *FoundIndex);
+			}
+			else
+			{
+				UE_LOG(
+					LogUsd,
+					Warning,
+					TEXT("Failed to find primvar '%s' when setting material parameter '%s' on material '%s'. Available primvars and UV "
+						 "indices: %s.%s"),
+					*TextureValue.Primvar,
+					ParameterName,
+					*Material.GetPathName(),
+					*UsdUtils::StringifyMap(PrimvarToUVIndex),
+					TextureValue.Primvar.IsEmpty() ? TEXT(" Is your UsdUVTexture Shader missing the 'inputs:st' attribute? (It specifies "
+														  "which UV set to sample the texture with)")
+												   : TEXT("")
+				);
 			}
 
-#if WITH_EDITOR
-			// Note that we will bake things that aren't supported on the default USD surface shader schema. These could be useful in case the user
-			// has a custom renderer though, and they can pick which properties they want anyway
-			bool BakeMaterial(
-				const UMaterialInterface& Material,
-				const TArray<FPropertyEntry>& InMaterialProperties,
-				const FIntPoint& InDefaultTextureSize,
-				FBakeOutput& OutBakedData,
-				bool bInDecayTexturesToSinglePixel
-			)
+			FLinearColor ComponentMask = FLinearColor::Black;
+			switch (TextureValue.OutputIndex)
 			{
-				const bool bAllQualityLevels = true;
-				const bool bAllFeatureLevels = true;
-				TArray<UTexture*> MaterialTextures;
-				Material.GetUsedTextures(MaterialTextures, EMaterialQualityLevel::Num, bAllQualityLevels, GMaxRHIFeatureLevel, bAllFeatureLevels);
+				case 0:	   // RGB
+					ComponentMask = FLinearColor{1.f, 1.f, 1.f, 0.f};
+					break;
+				case 1:	   // R
+					ComponentMask = FLinearColor{1.f, 0.f, 0.f, 0.f};
+					break;
+				case 2:	   // G
+					ComponentMask = FLinearColor{0.f, 1.f, 0.f, 0.f};
+					break;
+				case 3:	   // B
+					ComponentMask = FLinearColor{0.f, 0.f, 1.f, 0.f};
+					break;
+				case 4:	   // A
+					ComponentMask = FLinearColor{0.f, 0.f, 0.f, 1.f};
+					break;
+			}
+			UsdUtils::SetVectorParameterValue(Material, *FString::Printf(TEXT("%sTextureComponent"), ParameterName), ComponentMask);
+		}
 
-				// Precache all used textures, otherwise could get everything rendered with low-res textures.
-				for (UTexture* Texture : MaterialTextures)
-				{
-					if (UTexture2D* Texture2D = Cast<UTexture2D>(Texture))
-					{
-						Texture2D->SetForceMipLevelsToBeResident(30.0f);
-						Texture2D->WaitForStreaming();
-					}
-				}
+		void operator()(const UsdToUnreal::FPrimvarReaderParameterValue& PrimvarReaderValue) const
+		{
+			FSetParameterValueVisitor::operator()(PrimvarReaderValue.FallbackValue);
 
-				FMaterialData MatSet;
-				MatSet.Material = const_cast<UMaterialInterface*>(&Material);	 // We don't modify it and neither does the material baking module,
-																				 // it's just a bad signature
-				MatSet.bPerformShrinking = bInDecayTexturesToSinglePixel;
+			// We currently only import the "displayColor" primvar in USD as vertex colors
+			// on UE meshes, so that's the only primvar we can support here.
+			if (PrimvarReaderValue.PrimvarName == "displayColor")
+			{
+				UsdUtils::SetScalarParameterValue(Material, TEXT("UseVertexColorForBaseColor"), 1.0f);
+			}
+		}
 
-				for (const FPropertyEntry& Entry : InMaterialProperties)
-				{
-					// No point in asking it to bake if we're going to use the user-supplied value
-					if (Entry.bUseConstantValue)
-					{
-						continue;
-					}
+	protected:
+		const TMap<FString, int32>& PrimvarToUVIndex;
+	};
 
-					switch (Entry.Property)
-					{
-						case MP_Normal:
-							if (!Material.GetMaterial()->HasNormalConnected() && !Material.GetMaterial()->bUseMaterialAttributes)
-							{
-								continue;
-							}
-							break;
-						case MP_Tangent:
-							if (!Material.GetMaterial()->GetEditorOnlyData()->Tangent.IsConnected()
-								&& !Material.GetMaterial()->bUseMaterialAttributes)
-							{
-								continue;
-							}
-							break;
-						case MP_EmissiveColor:
-							if (!Material.GetMaterial()->GetEditorOnlyData()->EmissiveColor.IsConnected()
-								&& !Material.GetMaterial()->bUseMaterialAttributes)
-							{
-								continue;
-							}
-							break;
-						case MP_OpacityMask:
-							if (!Material.IsPropertyActive(MP_OpacityMask) || !IsMaskedBlendMode(Material))
-							{
-								continue;
-							}
-							break;
-						case MP_Opacity:
-							if (!Material.IsPropertyActive(MP_Opacity) || !IsTranslucentBlendMode(Material))
-							{
-								continue;
-							}
-							break;
-						case MP_MAX:
-							continue;
-							break;
-						default:
-							if (!Material.IsPropertyActive(Entry.Property))
-							{
-								continue;
-							}
-							break;
-					}
+	void SetParameterValue(
+		UMaterialInstance& Material,
+		const TCHAR* ParameterName,
+		const UsdToUnreal::FParameterValue& ParameterValue,
+		bool bForUsdPreviewSurface,
+		const TMap<FString, int32>& PrimvarToUVIndex
+	)
+	{
+		if (bForUsdPreviewSurface)
+		{
+			FSetPreviewSurfaceParameterValueVisitor Visitor{Material, ParameterName, PrimvarToUVIndex};
+			Visit(Visitor, ParameterValue);
+		}
+		else
+		{
+			FSetParameterValueVisitor Visitor{Material, ParameterName};
+			Visit(Visitor, ParameterValue);
+		}
+	}
 
-					MatSet.PropertySizes.Add(Entry.Property, Entry.bUseCustomSize ? Entry.CustomSize : InDefaultTextureSize);
-				}
+	UTexture* CreateTextureWithEditor(
+		const FString& ResolvedTexturePath,
+		FName SanitizedName,
+		TextureGroup Group,
+		EObjectFlags ObjectFlags,
+		UObject* Outer,
+		bool bForceLinear
+	)
+	{
+		UTexture* Texture = nullptr;
+#if WITH_EDITOR
+		if (ResolvedTexturePath.IsEmpty() || !Outer)
+		{
+			return nullptr;
+		}
 
-				FMeshData MeshSettings;
-				MeshSettings.MeshDescription = nullptr;
-				MeshSettings.TextureCoordinateBox = FBox2D(FVector2D(0.0f, 0.0f), FVector2D(1.0f, 1.0f));
-				MeshSettings.TextureCoordinateIndex = 0;
+		FScopedUnrealAllocs UnrealAllocs;
 
-				TArray<FBakeOutput> BakeOutputs;
-				IMaterialBakingModule& Module = FModuleManager::Get().LoadModuleChecked<IMaterialBakingModule>("MaterialBaking");
-				const bool bLinearBake = true;
-				Module.SetLinearBake(bLinearBake);
-				const bool bEmissiveHDR = true;
-				Module.SetEmissiveHDR(bEmissiveHDR);
-				Module.BakeMaterials({&MatSet}, {&MeshSettings}, BakeOutputs);
+		UTextureFactory* TextureFactory = NewObject<UTextureFactory>();
+		TextureFactory->SuppressImportOverwriteDialog();
+		TextureFactory->bUseHashAsGuid = true;
+		TextureFactory->LODGroup = Group;
+		TextureFactory->HDRImportShouldBeLongLatCubeMap = EAppReturnType::YesAll;
+		// To maintain existing behavior, ColorSpaceMode is left at auto if linear is not requested
+		TextureFactory->ColorSpaceMode = bForceLinear ? ETextureSourceColorSpace::Linear : ETextureSourceColorSpace::Auto;
 
-				if (BakeOutputs.Num() < 1)
-				{
-					return false;
-				}
+		const bool bIsSupportedUdimTexture = ResolvedTexturePath.Contains(TEXT("<UDIM>"));
+		if (bIsSupportedUdimTexture)
+		{
+			FString BaseFileName = FPaths::GetBaseFilename(ResolvedTexturePath);
 
-				OutBakedData = BakeOutputs[0];
-				return true;
+			FString BaseFileNameBeforeUdim;
+			FString BaseFileNameAfterUdim;
+			BaseFileName.Split(TEXT("<UDIM>"), &BaseFileNameBeforeUdim, &BaseFileNameAfterUdim);
+
+			FString UdimRegexPattern = FString::Printf(TEXT(R"((%s)(\d{4})(%s))"), *BaseFileNameBeforeUdim, *BaseFileNameAfterUdim);
+			TextureFactory->UdimRegexPattern = MoveTemp(UdimRegexPattern);
+		}
+
+		if (!ResolvedTexturePath.IsEmpty())
+		{
+			FString TextureExtension;
+			if (IsInsideUsdzArchive(ResolvedTexturePath, TextureExtension))
+			{
+				// Always prefer using the TextureFactory if we can, as it may provide compression, which the runtime version never will
+				Texture = ReadTextureFromUsdzArchiveEditor(ResolvedTexturePath, TextureExtension, TextureFactory, Outer, ObjectFlags);
+			}
+			// Not inside an USDZ archive, just a regular texture
+			else
+			{
+				bool bOutCancelled = false;
+				const TCHAR* Params = TEXT("");
+				Texture = Cast<UTexture>(TextureFactory->ImportObject(	  //
+					UTexture::StaticClass(),
+					Outer,
+					SanitizedName,
+					ObjectFlags,
+					ResolvedTexturePath,
+					Params,
+					bOutCancelled
+				));
 			}
 
-			// Writes textures for all baked channels in BakedSamples that are larger than 1x1, and returns the filenames of the emitted textures for
-			// each channel
-			TMap<EMaterialProperty, FString> WriteTextures(
-				FBakedMaterialView& BakedSamples,
-				const FString& TextureNamePrefix,
-				const FDirectoryPath& TexturesFolder
-			)
+			if (Texture)
 			{
-				IImageWrapperModule& ImageWrapperModule = FModuleManager::Get().LoadModuleChecked<IImageWrapperModule>(TEXT("ImageWrapper"));
-				TSharedPtr<IImageWrapper> EXRImageWrapper = ImageWrapperModule.CreateImageWrapper(EImageFormat::EXR);
-				TSharedPtr<IImageWrapper> PNGImageWrapper = ImageWrapperModule.CreateImageWrapper(EImageFormat::PNG);
+				// We set this even if we're not going to import so that we can track our original texture filepath
+				// in case we later do an Actions->Import
+				UUsdAssetImportData* ImportData = NewObject<UUsdAssetImportData>(Texture);
+				ImportData->UpdateFilenameOnly(ResolvedTexturePath);
+				Texture->AssetImportData = ImportData;
+			}
+		}
+#endif	  // WITH_EDITOR
+		return Texture;
+	}
 
-				TMap<EMaterialProperty, FString> WrittenTexturesPerChannel;
+	UTexture* CreateTextureAtRuntime(const FString& ResolvedTexturePath)
+	{
+		FScopedUnrealAllocs UnrealAllocs;
 
-				// Write textures for HDR baked properties larger than 1x1
-				for (TPair<EMaterialProperty, TArray<FFloat16Color>*>& BakedDataPair : BakedSamples.HDRPropertyData)
-				{
-					const EMaterialProperty Property = BakedDataPair.Key;
+		UTexture* Texture = nullptr;
 
-					TArray<FFloat16Color>* Samples = BakedDataPair.Value;
-					if (!Samples || Samples->Num() < 2)
-					{
-						continue;
-					}
+		if (!ResolvedTexturePath.IsEmpty())
+		{
+			// Try checking if the texture is inside an USDZ archive first, or else TextureFactory throws an error
+			FString TextureExtension;
+			if (IsInsideUsdzArchive(ResolvedTexturePath, TextureExtension))
+			{
+				Texture = ReadTextureFromUsdzArchiveRuntime(ResolvedTexturePath);
+			}
 
-					const FIntPoint& FinalSize = BakedSamples.PropertySizes.FindChecked(Property);
-					if (FinalSize.GetMin() < 2)
-					{
-						continue;
-					}
+			// Not inside an USDZ archive, just a regular texture
+			if (!Texture)
+			{
+				Texture = FImageUtils::ImportFileAsTexture2D(ResolvedTexturePath);
+			}
+		}
 
-					const UEnum* PropertyEnum = StaticEnum<EMaterialProperty>();
-					FName PropertyName = PropertyEnum->GetNameByValue(Property);
-					FString TrimmedPropertyName = PropertyName.ToString();
-					TrimmedPropertyName.RemoveFromStart(TEXT("MP_"));
-
-					// Final FilePath will be something like "C:/TexturesFolder/Game_ContentFolder_Materials_Red_BaseColor.png", which automatically
-					// guarantees it won't overwrite another texture from the same export, but will overwrite old textures from previous exports
-					FString TextureFileName = UsdUnreal::ObjectUtils::SanitizeObjectName(FPaths::ChangeExtension(TextureNamePrefix, TEXT("")));
-					TextureFileName.RemoveFromStart(TEXT("_"));
-					FString TextureFilePath = FPaths::Combine(
-						TexturesFolder.Path,
-						FString::Printf(TEXT("%s_%s.exr"), *UsdUnreal::ObjectUtils::SanitizeObjectName(TextureFileName), *TrimmedPropertyName)
-					);
+		return Texture;
+	}
 
-					// For some reason the baked samples always have zero alpha and there is nothing we can do about it... It seems like the material
-					// baking module is made with the intent that the data ends up in UTexture2Ds, where they can be set to be compressed without
-					// alpha and have the value ignored. Since we need to write these to file, we must set them back up to full alpha. This is
-					// potentially useless as USD handles these at most as color3f, but it could be annoying for the user if they intend on using the
-					// textures for anything else
-					for (FFloat16Color& Sample : *Samples)
-					{
-						Sample.A = 1.0f;
-					}
+#if WITH_EDITOR
+	// Note that we will bake things that aren't supported on the default USD surface shader schema. These could be useful in case the user
+	// has a custom renderer though, and they can pick which properties they want anyway
+	bool BakeMaterial(
+		const UMaterialInterface& Material,
+		const TArray<FPropertyEntry>& InMaterialProperties,
+		const FIntPoint& InDefaultTextureSize,
+		FBakeOutput& OutBakedData,
+		bool bInDecayTexturesToSinglePixel
+	)
+	{
+		const bool bAllQualityLevels = true;
+		const bool bAllFeatureLevels = true;
+		TArray<UTexture*> MaterialTextures;
+		Material.GetUsedTextures(MaterialTextures, EMaterialQualityLevel::Num, bAllQualityLevels, GMaxRHIFeatureLevel, bAllFeatureLevels);
+
+		// Precache all used textures, otherwise could get everything rendered with low-res textures.
+		for (UTexture* Texture : MaterialTextures)
+		{
+			if (UTexture2D* Texture2D = Cast<UTexture2D>(Texture))
+			{
+				Texture2D->SetForceMipLevelsToBeResident(30.0f);
+				Texture2D->WaitForStreaming();
+			}
+		}
+
+		FMaterialData MatSet;
+		MatSet.Material = const_cast<UMaterialInterface*>(&Material);	 // We don't modify it and neither does the material baking module,
+																		 // it's just a bad signature
+		MatSet.bPerformShrinking = bInDecayTexturesToSinglePixel;
 
-					EXRImageWrapper->SetRaw(Samples->GetData(), Samples->GetAllocatedSize(), FinalSize.X, FinalSize.Y, ERGBFormat::RGBAF, 16);
-					const TArray64<uint8> Data = EXRImageWrapper->GetCompressed(100);
+		for (const FPropertyEntry& Entry : InMaterialProperties)
+		{
+			// No point in asking it to bake if we're going to use the user-supplied value
+			if (Entry.bUseConstantValue)
+			{
+				continue;
+			}
 
-					bool bWroteFile = FFileHelper::SaveArrayToFile(Data, *TextureFilePath);
-					if (bWroteFile)
+			switch (Entry.Property)
+			{
+				case MP_Normal:
+					if (!Material.GetMaterial()->HasNormalConnected() && !Material.GetMaterial()->bUseMaterialAttributes)
 					{
-						WrittenTexturesPerChannel.Add(Property, TextureFilePath);
+						continue;
 					}
-					else
+					break;
+				case MP_Tangent:
+					if (!Material.GetMaterial()->GetEditorOnlyData()->Tangent.IsConnected() && !Material.GetMaterial()->bUseMaterialAttributes)
 					{
-						UE_LOG(LogUsd, Warning, TEXT("Failed to write texture '%s', baked channel will be ignored."), *TextureFilePath);
+						continue;
 					}
-				}
-
-				// Write textures for baked properties larger than 1x1
-				for (TPair<EMaterialProperty, TArray<FColor>*>& BakedDataPair : BakedSamples.PropertyData)
-				{
-					const EMaterialProperty Property = BakedDataPair.Key;
-
-					// The material baking module still generates and sends an SDR version of any HDR channel it also bakes,
-					// so here let's skip this one in case we already generated an HDR texture for the channel
-					if (WrittenTexturesPerChannel.Contains(Property))
+					break;
+				case MP_EmissiveColor:
+					if (!Material.GetMaterial()->GetEditorOnlyData()->EmissiveColor.IsConnected() && !Material.GetMaterial()->bUseMaterialAttributes)
 					{
 						continue;
 					}
-
-					TArray<FColor>* Samples = BakedDataPair.Value;
-					if (!Samples || Samples->Num() < 2)
+					break;
+				case MP_OpacityMask:
+					if (!Material.IsPropertyActive(MP_OpacityMask) || !IsMaskedBlendMode(Material))
 					{
 						continue;
 					}
-
-					const FIntPoint& FinalSize = BakedSamples.PropertySizes.FindChecked(Property);
-					if (FinalSize.GetMin() < 2)
+					break;
+				case MP_Opacity:
+					if (!Material.IsPropertyActive(MP_Opacity) || !IsTranslucentBlendMode(Material))
 					{
 						continue;
 					}
-
-					const UEnum* PropertyEnum = StaticEnum<EMaterialProperty>();
-					FName PropertyName = PropertyEnum->GetNameByValue(Property);
-					FString TrimmedPropertyName = PropertyName.ToString();
-					TrimmedPropertyName.RemoveFromStart(TEXT("MP_"));
-
-					// Final FilePath will be something like "C:/TexturesFolder/Game_ContentFolder_Materials_Red_BaseColor.png", which automatically
-					// guarantees it won't overwrite another texture from the same export, but will overwrite old textures from previous exports
-					FString TextureFileName = UsdUnreal::ObjectUtils::SanitizeObjectName(FPaths::ChangeExtension(TextureNamePrefix, TEXT("")));
-					TextureFileName.RemoveFromStart(TEXT("_"));
-					FString TextureFilePath = FPaths::Combine(
-						TexturesFolder.Path,
-						FString::Printf(TEXT("%s_%s.png"), *UsdUnreal::ObjectUtils::SanitizeObjectName(TextureFileName), *TrimmedPropertyName)
-					);
-
-					// For some reason the baked samples always have zero alpha and there is nothing we can do about it... It seems like the material
-					// baking module is made with the intent that the data ends up in UTexture2Ds, where they can be set to be compressed without
-					// alpha and have the value ignored. Since we need to write these to file, we must set them back up to full alpha. This is
-					// potentially useless as USD handles these at most as color3f, but it could be annoying for the user if they intend on using the
-					// textures for anything else
-					for (FColor& Sample : *Samples)
+					break;
+				case MP_MAX:
+					continue;
+					break;
+				default:
+					if (!Material.IsPropertyActive(Entry.Property))
 					{
-						Sample.A = 255;
+						continue;
 					}
+					break;
+			}
 
-					PNGImageWrapper->SetRaw(Samples->GetData(), Samples->GetAllocatedSize(), FinalSize.X, FinalSize.Y, ERGBFormat::BGRA, 8);
-					const TArray64<uint8> PNGData = PNGImageWrapper->GetCompressed(100);
+			MatSet.PropertySizes.Add(Entry.Property, Entry.bUseCustomSize ? Entry.CustomSize : InDefaultTextureSize);
+		}
 
-					bool bWroteFile = FFileHelper::SaveArrayToFile(PNGData, *TextureFilePath);
-					if (bWroteFile)
-					{
-						WrittenTexturesPerChannel.Add(Property, TextureFilePath);
-					}
-					else
-					{
-						UE_LOG(LogUsd, Warning, TEXT("Failed to write texture '%s', baked channel will be ignored."), *TextureFilePath);
-					}
-				}
+		FMeshData MeshSettings;
+		MeshSettings.MeshDescription = nullptr;
+		MeshSettings.TextureCoordinateBox = FBox2D(FVector2D(0.0f, 0.0f), FVector2D(1.0f, 1.0f));
+		MeshSettings.TextureCoordinateIndex = 0;
+
+		TArray<FBakeOutput> BakeOutputs;
+		IMaterialBakingModule& Module = FModuleManager::Get().LoadModuleChecked<IMaterialBakingModule>("MaterialBaking");
+		const bool bLinearBake = true;
+		Module.SetLinearBake(bLinearBake);
+		const bool bEmissiveHDR = true;
+		Module.SetEmissiveHDR(bEmissiveHDR);
+		Module.BakeMaterials({&MatSet}, {&MeshSettings}, BakeOutputs);
+
+		if (BakeOutputs.Num() < 1)
+		{
+			return false;
+		}
+
+		OutBakedData = BakeOutputs[0];
+		return true;
+	}
+
+	// Writes textures for all baked channels in BakedSamples that are larger than 1x1, and returns the filenames of the emitted textures for
+	// each channel
+	TMap<EMaterialProperty, FString> WriteTextures(
+		FBakedMaterialView& BakedSamples,
+		const FString& TextureNamePrefix,
+		const FDirectoryPath& TexturesFolder
+	)
+	{
+		IImageWrapperModule& ImageWrapperModule = FModuleManager::Get().LoadModuleChecked<IImageWrapperModule>(TEXT("ImageWrapper"));
+		TSharedPtr<IImageWrapper> EXRImageWrapper = ImageWrapperModule.CreateImageWrapper(EImageFormat::EXR);
+		TSharedPtr<IImageWrapper> PNGImageWrapper = ImageWrapperModule.CreateImageWrapper(EImageFormat::PNG);
+
+		TMap<EMaterialProperty, FString> WrittenTexturesPerChannel;
+
+		// Write textures for HDR baked properties larger than 1x1
+		for (TPair<EMaterialProperty, TArray<FFloat16Color>*>& BakedDataPair : BakedSamples.HDRPropertyData)
+		{
+			const EMaterialProperty Property = BakedDataPair.Key;
 
-				return WrittenTexturesPerChannel;
+			TArray<FFloat16Color>* Samples = BakedDataPair.Value;
+			if (!Samples || Samples->Num() < 2)
+			{
+				continue;
 			}
 
-			bool ConfigureShadePrim(
-				const FBakedMaterialView& BakedData,
-				const TMap<EMaterialProperty, FString>& WrittenTextures,
-				const TMap<EMaterialProperty, float>& UserConstantValues,
-				pxr::UsdShadeMaterial& OutUsdShadeMaterial
-			)
+			const FIntPoint& FinalSize = BakedSamples.PropertySizes.FindChecked(Property);
+			if (FinalSize.GetMin() < 2)
 			{
-				FScopedUsdAllocs Allocs;
+				continue;
+			}
 
-				pxr::UsdPrim MaterialPrim = OutUsdShadeMaterial.GetPrim();
-				pxr::UsdStageRefPtr Stage = MaterialPrim.GetStage();
-				if (!MaterialPrim || !Stage)
-				{
-					return false;
-				}
+			const UEnum* PropertyEnum = StaticEnum<EMaterialProperty>();
+			FName PropertyName = PropertyEnum->GetNameByValue(Property);
+			FString TrimmedPropertyName = PropertyName.ToString();
+			TrimmedPropertyName.RemoveFromStart(TEXT("MP_"));
+
+			// Final FilePath will be something like "C:/TexturesFolder/Game_ContentFolder_Materials_Red_BaseColor.png", which automatically
+			// guarantees it won't overwrite another texture from the same export, but will overwrite old textures from previous exports
+			FString TextureFileName = UsdUnreal::ObjectUtils::SanitizeObjectName(FPaths::ChangeExtension(TextureNamePrefix, TEXT("")));
+			TextureFileName.RemoveFromStart(TEXT("_"));
+			FString TextureFilePath = FPaths::Combine(
+				TexturesFolder.Path,
+				FString::Printf(TEXT("%s_%s.exr"), *UsdUnreal::ObjectUtils::SanitizeObjectName(TextureFileName), *TrimmedPropertyName)
+			);
+
+			// For some reason the baked samples always have zero alpha and there is nothing we can do about it... It seems like the material
+			// baking module is made with the intent that the data ends up in UTexture2Ds, where they can be set to be compressed without
+			// alpha and have the value ignored. Since we need to write these to file, we must set them back up to full alpha. This is
+			// potentially useless as USD handles these at most as color3f, but it could be annoying for the user if they intend on using the
+			// textures for anything else
+			for (FFloat16Color& Sample : *Samples)
+			{
+				Sample.A = 1.0f;
+			}
+
+			EXRImageWrapper->SetRaw(Samples->GetData(), Samples->GetAllocatedSize(), FinalSize.X, FinalSize.Y, ERGBFormat::RGBAF, 16);
+			const TArray64<uint8> Data = EXRImageWrapper->GetCompressed(100);
 
-				FString UsdFilePath = UsdToUnreal::ConvertString(Stage->GetRootLayer()->GetRealPath());
+			bool bWroteFile = FFileHelper::SaveArrayToFile(Data, *TextureFilePath);
+			if (bWroteFile)
+			{
+				WrittenTexturesPerChannel.Add(Property, TextureFilePath);
+			}
+			else
+			{
+				UE_LOG(LogUsd, Warning, TEXT("Failed to write texture '%s', baked channel will be ignored."), *TextureFilePath);
+			}
+		}
 
-				FUsdStageInfo StageInfo{Stage};
+		// Write textures for baked properties larger than 1x1
+		for (TPair<EMaterialProperty, TArray<FColor>*>& BakedDataPair : BakedSamples.PropertyData)
+		{
+			const EMaterialProperty Property = BakedDataPair.Key;
 
-				pxr::SdfPath MaterialPath = MaterialPrim.GetPath();
+			// The material baking module still generates and sends an SDR version of any HDR channel it also bakes,
+			// so here let's skip this one in case we already generated an HDR texture for the channel
+			if (WrittenTexturesPerChannel.Contains(Property))
+			{
+				continue;
+			}
 
-				// Create surface shader
-				pxr::UsdShadeShader ShadeShader = pxr::UsdShadeShader::Define(
-					Stage,
-					MaterialPath.AppendChild(UnrealToUsd::ConvertToken(TEXT("SurfaceShader")).Get())
-				);
-				ShadeShader.SetShaderId(UnrealIdentifiers::UsdPreviewSurface);
-				pxr::UsdShadeOutput ShaderOutOutput = ShadeShader.CreateOutput(UnrealIdentifiers::Surface, pxr::SdfValueTypeNames->Token);
+			TArray<FColor>* Samples = BakedDataPair.Value;
+			if (!Samples || Samples->Num() < 2)
+			{
+				continue;
+			}
 
-				// Connect material to surface shader
-				pxr::UsdShadeOutput MaterialSurfaceOutput = OutUsdShadeMaterial.CreateSurfaceOutput();
-				MaterialSurfaceOutput.ConnectToSource(ShaderOutOutput);
+			const FIntPoint& FinalSize = BakedSamples.PropertySizes.FindChecked(Property);
+			if (FinalSize.GetMin() < 2)
+			{
+				continue;
+			}
 
-				pxr::UsdShadeShader PrimvarReaderShader;
+			const UEnum* PropertyEnum = StaticEnum<EMaterialProperty>();
+			FName PropertyName = PropertyEnum->GetNameByValue(Property);
+			FString TrimmedPropertyName = PropertyName.ToString();
+			TrimmedPropertyName.RemoveFromStart(TEXT("MP_"));
+
+			// Final FilePath will be something like "C:/TexturesFolder/Game_ContentFolder_Materials_Red_BaseColor.png", which automatically
+			// guarantees it won't overwrite another texture from the same export, but will overwrite old textures from previous exports
+			FString TextureFileName = UsdUnreal::ObjectUtils::SanitizeObjectName(FPaths::ChangeExtension(TextureNamePrefix, TEXT("")));
+			TextureFileName.RemoveFromStart(TEXT("_"));
+			FString TextureFilePath = FPaths::Combine(
+				TexturesFolder.Path,
+				FString::Printf(TEXT("%s_%s.png"), *UsdUnreal::ObjectUtils::SanitizeObjectName(TextureFileName), *TrimmedPropertyName)
+			);
+
+			// For some reason the baked samples always have zero alpha and there is nothing we can do about it... It seems like the material
+			// baking module is made with the intent that the data ends up in UTexture2Ds, where they can be set to be compressed without
+			// alpha and have the value ignored. Since we need to write these to file, we must set them back up to full alpha. This is
+			// potentially useless as USD handles these at most as color3f, but it could be annoying for the user if they intend on using the
+			// textures for anything else
+			for (FColor& Sample : *Samples)
+			{
+				Sample.A = 255;
+			}
 
-				const pxr::TfToken PrimvarReaderShaderName = UnrealToUsd::ConvertToken(TEXT("PrimvarReader")).Get();
-				const pxr::TfToken PrimvarVariableName = UnrealToUsd::ConvertToken(TEXT("stPrimvarName")).Get();
+			PNGImageWrapper->SetRaw(Samples->GetData(), Samples->GetAllocatedSize(), FinalSize.X, FinalSize.Y, ERGBFormat::BGRA, 8);
+			const TArray64<uint8> PNGData = PNGImageWrapper->GetCompressed(100);
 
-				// Collect all the properties we'll process, as some data is baked and some comes from values the user input as export options
-				TSet<EMaterialProperty> PropertiesToProcess;
-				{
-					PropertiesToProcess.Reserve(BakedData.PropertyData.Num() + UserConstantValues.Num() + BakedData.HDRPropertyData.Num());
-					BakedData.PropertyData.GetKeys(PropertiesToProcess);
+			bool bWroteFile = FFileHelper::SaveArrayToFile(PNGData, *TextureFilePath);
+			if (bWroteFile)
+			{
+				WrittenTexturesPerChannel.Add(Property, TextureFilePath);
+			}
+			else
+			{
+				UE_LOG(LogUsd, Warning, TEXT("Failed to write texture '%s', baked channel will be ignored."), *TextureFilePath);
+			}
+		}
 
-					TSet<EMaterialProperty> UsedProperties;
-					UserConstantValues.GetKeys(UsedProperties);
-					PropertiesToProcess.Append(UsedProperties);
+		return WrittenTexturesPerChannel;
+	}
 
-					BakedData.HDRPropertyData.GetKeys(UsedProperties);
-					PropertiesToProcess.Append(UsedProperties);
-				}
+	bool ConfigureShadePrim(
+		const FBakedMaterialView& BakedData,
+		const TMap<EMaterialProperty, FString>& WrittenTextures,
+		const TMap<EMaterialProperty, float>& UserConstantValues,
+		pxr::UsdShadeMaterial& OutUsdShadeMaterial
+	)
+	{
+		FScopedUsdAllocs Allocs;
+
+		pxr::UsdPrim MaterialPrim = OutUsdShadeMaterial.GetPrim();
+		pxr::UsdStageRefPtr Stage = MaterialPrim.GetStage();
+		if (!MaterialPrim || !Stage)
+		{
+			return false;
+		}
+
+		FString UsdFilePath = UsdToUnreal::ConvertString(Stage->GetRootLayer()->GetRealPath());
 
-				// We always write BaseColor because the default in UE is full black, but in USD seems to be 0.18. If we left
-				// BaseColor unbound, a material that relies on opacity/other channels and leaves BaseColor disconnected for the black
-				// value would end up looking gray in usdview/other DCCs
-				if (!PropertiesToProcess.Contains(MP_BaseColor))
+		FUsdStageInfo StageInfo{Stage};
+
+		pxr::SdfPath MaterialPath = MaterialPrim.GetPath();
+
+		// Create surface shader
+		pxr::UsdShadeShader ShadeShader = pxr::UsdShadeShader::Define(
+			Stage,
+			MaterialPath.AppendChild(UnrealToUsd::ConvertToken(TEXT("SurfaceShader")).Get())
+		);
+		ShadeShader.SetShaderId(UnrealIdentifiers::UsdPreviewSurface);
+		pxr::UsdShadeOutput ShaderOutOutput = ShadeShader.CreateOutput(UnrealIdentifiers::Surface, pxr::SdfValueTypeNames->Token);
+
+		// Connect material to surface shader
+		pxr::UsdShadeOutput MaterialSurfaceOutput = OutUsdShadeMaterial.CreateSurfaceOutput();
+		MaterialSurfaceOutput.ConnectToSource(ShaderOutOutput);
+
+		pxr::UsdShadeShader PrimvarReaderShader;
+
+		const pxr::TfToken PrimvarReaderShaderName = UnrealToUsd::ConvertToken(TEXT("PrimvarReader")).Get();
+		const pxr::TfToken PrimvarVariableName = UnrealToUsd::ConvertToken(TEXT("stPrimvarName")).Get();
+
+		// Collect all the properties we'll process, as some data is baked and some comes from values the user input as export options
+		TSet<EMaterialProperty> PropertiesToProcess;
+		{
+			PropertiesToProcess.Reserve(BakedData.PropertyData.Num() + UserConstantValues.Num() + BakedData.HDRPropertyData.Num());
+			BakedData.PropertyData.GetKeys(PropertiesToProcess);
+
+			TSet<EMaterialProperty> UsedProperties;
+			UserConstantValues.GetKeys(UsedProperties);
+			PropertiesToProcess.Append(UsedProperties);
+
+			BakedData.HDRPropertyData.GetKeys(UsedProperties);
+			PropertiesToProcess.Append(UsedProperties);
+		}
+
+		// We always write BaseColor because the default in UE is full black, but in USD seems to be 0.18. If we left
+		// BaseColor unbound, a material that relies on opacity/other channels and leaves BaseColor disconnected for the black
+		// value would end up looking gray in usdview/other DCCs
+		if (!PropertiesToProcess.Contains(MP_BaseColor))
+		{
+			PropertiesToProcess.Add(MP_BaseColor);
+		}
+		const float Zero = 0.0f;
+
+		// Fill in outputs
+		for (const EMaterialProperty& Property : PropertiesToProcess)
+		{
+			const FString* TextureFilePath = WrittenTextures.Find(Property);
+			const float* UserConstantValue = UserConstantValues.Find(Property);
+			const FIntPoint* SampleSize = BakedData.PropertySizes.Find(Property);
+
+			uint64 NumSamples = 0;
+
+			bool bPropertyValueIsConstant = false;
+			pxr::GfVec3f ConstantLinearValue;
+
+			// Try HDR first: When baking a channel as HDR (like emissive) the MaterialBaking module
+			// will still bake the SDR version of the channel too. We keep both in our FBakedMaterialView
+			// because it only does the mechanism of decaying to a single value on the SDR data array
+			bool bParsedHDR = false;
+			if (BakedData.HDRPropertyData.Contains(Property))
+			{
+				const TArray<FColor>* SDRSamples = BakedData.PropertyData.FindRef(Property);
+				NumSamples = SDRSamples ? SDRSamples->Num() : 0;
+
+				const bool bDecayedToSingleSample = NumSamples == 1;
+				bPropertyValueIsConstant = (UserConstantValue != nullptr || bDecayedToSingleSample);
+
+				const TArray<FFloat16Color>* Samples = BakedData.HDRPropertyData.FindRef(Property);
+				if (bDecayedToSingleSample && Samples)
 				{
-					PropertiesToProcess.Add(MP_BaseColor);
+					// If it decayed to single sample we know our SDR array only has one value,
+					// and so should our HDR array. Get the value from the HDR one to avoid an
+					// unnecessary quantization though
+					const FFloat16Color& Sample = (*Samples)[0];
+					ConstantLinearValue = pxr::GfVec3f{Sample.R, Sample.G, Sample.B};
 				}
-				const float Zero = 0.0f;
 
-				// Fill in outputs
-				for (const EMaterialProperty& Property : PropertiesToProcess)
+				if (Samples && Samples->Num() > 0)
 				{
-					const FString* TextureFilePath = WrittenTextures.Find(Property);
-					const float* UserConstantValue = UserConstantValues.Find(Property);
-					const FIntPoint* SampleSize = BakedData.PropertySizes.Find(Property);
-
-					uint64 NumSamples = 0;
+					bParsedHDR = true;
+				}
+			}
 
-					bool bPropertyValueIsConstant = false;
-					pxr::GfVec3f ConstantLinearValue;
+			if (!bParsedHDR)
+			{
+				const TArray<FColor>* Samples = BakedData.PropertyData.FindRef(Property);
+				NumSamples = Samples ? Samples->Num() : 0;
 
-					// Try HDR first: When baking a channel as HDR (like emissive) the MaterialBaking module
-					// will still bake the SDR version of the channel too. We keep both in our FBakedMaterialView
-					// because it only does the mechanism of decaying to a single value on the SDR data array
-					bool bParsedHDR = false;
-					if (BakedData.HDRPropertyData.Contains(Property))
+				bPropertyValueIsConstant = (UserConstantValue != nullptr || NumSamples == 1);
+				if (NumSamples == 1)
+				{
+					switch (Property)
 					{
-						const TArray<FColor>* SDRSamples = BakedData.PropertyData.FindRef(Property);
-						NumSamples = SDRSamples ? SDRSamples->Num() : 0;
-
-						const bool bDecayedToSingleSample = NumSamples == 1;
-						bPropertyValueIsConstant = (UserConstantValue != nullptr || bDecayedToSingleSample);
-
-						const TArray<FFloat16Color>* Samples = BakedData.HDRPropertyData.FindRef(Property);
-						if (bDecayedToSingleSample && Samples)
+						case EMaterialProperty::MP_BaseColor:
+						case EMaterialProperty::MP_SubsurfaceColor:
 						{
-							// If it decayed to single sample we know our SDR array only has one value,
-							// and so should our HDR array. Get the value from the HDR one to avoid an
-							// unnecessary quantization though
-							const FFloat16Color& Sample = (*Samples)[0];
-							ConstantLinearValue = pxr::GfVec3f{Sample.R, Sample.G, Sample.B};
+							pxr::GfVec4f ConvertedColor = UnrealToUsd::ConvertColor((*Samples)[0]);
+							ConstantLinearValue = pxr::GfVec3f{ConvertedColor[0], ConvertedColor[1], ConvertedColor[2]};
+							break;
 						}
-
-						if (Samples && Samples->Num() > 0)
+						case EMaterialProperty::MP_Normal:
+						case EMaterialProperty::MP_Tangent:
 						{
-							bParsedHDR = true;
+							FVector ConvertedNormal{(*Samples)[0].ReinterpretAsLinear()};
+							ConstantLinearValue = UnrealToUsd::ConvertVectorFloat(StageInfo, ConvertedNormal);
+							break;
 						}
-					}
-
-					if (!bParsedHDR)
-					{
-						const TArray<FColor>* Samples = BakedData.PropertyData.FindRef(Property);
-						NumSamples = Samples ? Samples->Num() : 0;
-
-						bPropertyValueIsConstant = (UserConstantValue != nullptr || NumSamples == 1);
-						if (NumSamples == 1)
+						default:
 						{
-							switch (Property)
-							{
-								case EMaterialProperty::MP_BaseColor:
-								case EMaterialProperty::MP_SubsurfaceColor:
-								{
-									pxr::GfVec4f ConvertedColor = UnrealToUsd::ConvertColor((*Samples)[0]);
-									ConstantLinearValue = pxr::GfVec3f{ConvertedColor[0], ConvertedColor[1], ConvertedColor[2]};
-									break;
-								}
-								case EMaterialProperty::MP_Normal:
-								case EMaterialProperty::MP_Tangent:
-								{
-									FVector ConvertedNormal{(*Samples)[0].ReinterpretAsLinear()};
-									ConstantLinearValue = UnrealToUsd::ConvertVectorFloat(StageInfo, ConvertedNormal);
-									break;
-								}
-								default:
-								{
-									const FColor& Sample = (*Samples)[0];
-									ConstantLinearValue = pxr::GfVec3f{Sample.R / 255.0f, Sample.G / 255.0f, Sample.B / 255.0f};
-									break;
-								}
-							}
+							const FColor& Sample = (*Samples)[0];
+							ConstantLinearValue = pxr::GfVec3f{Sample.R / 255.0f, Sample.G / 255.0f, Sample.B / 255.0f};
+							break;
 						}
 					}
+				}
+			}
 
-					if (Property == MP_BaseColor && NumSamples == 0 && !UserConstantValue)
-					{
-						UserConstantValue = &Zero;
-						bPropertyValueIsConstant = true;
-					}
+			if (Property == MP_BaseColor && NumSamples == 0 && !UserConstantValue)
+			{
+				UserConstantValue = &Zero;
+				bPropertyValueIsConstant = true;
+			}
 
-					if ((NumSamples == 0 || !SampleSize) && !UserConstantValue)
-					{
-						UE_LOG(LogUsd, Log, TEXT("Skipping material property %d as we have no valid data to use."), Property);
-						continue;
-					}
+			if ((NumSamples == 0 || !SampleSize) && !UserConstantValue)
+			{
+				UE_LOG(LogUsd, Log, TEXT("Skipping material property %d as we have no valid data to use."), Property);
+				continue;
+			}
+
+			if (!UserConstantValue && NumSamples > 0 && SampleSize)
+			{
+				if (NumSamples != SampleSize->X * SampleSize->Y)
+				{
+					UE_LOG(
+						LogUsd,
+						Warning,
+						TEXT("Skipping material property %d as it has an unexpected number of samples (%d instead of %d)."),
+						Property,
+						NumSamples,
+						SampleSize->X * SampleSize->Y
+					);
+					continue;
+				}
+			}
+
+			if (!bPropertyValueIsConstant && (TextureFilePath == nullptr || !FPaths::FileExists(*TextureFilePath)))
+			{
+				UE_LOG(
+					LogUsd,
+					Warning,
+					TEXT("Skipping material property %d as target texture '%s' was not found."),
+					Property,
+					TextureFilePath ? **TextureFilePath : TEXT("")
+				);
+				continue;
+			}
 
-					if (!UserConstantValue && NumSamples > 0 && SampleSize)
+			pxr::TfToken InputToken;
+			pxr::SdfValueTypeName InputType;
+			pxr::VtValue ConstantValue;
+			pxr::GfVec4f FallbackValue;
+			pxr::TfToken ColorSpaceToken = UnrealIdentifiers::RawColorSpaceToken;
+			switch (Property)
+			{
+				case MP_BaseColor:
+					InputToken = UnrealIdentifiers::DiffuseColor;
+					InputType = pxr::SdfValueTypeNames->Color3f;
+					if (bPropertyValueIsConstant)
 					{
-						if (NumSamples != SampleSize->X * SampleSize->Y)
+						if (UserConstantValue)
 						{
-							UE_LOG(
-								LogUsd,
-								Warning,
-								TEXT("Skipping material property %d as it has an unexpected number of samples (%d instead of %d)."),
-								Property,
-								NumSamples,
-								SampleSize->X * SampleSize->Y
-							);
-							continue;
+							ConstantValue = pxr::GfVec3f(*UserConstantValue, *UserConstantValue, *UserConstantValue);
+						}
+						else
+						{
+							ConstantValue = ConstantLinearValue;
 						}
 					}
-
-					if (!bPropertyValueIsConstant && (TextureFilePath == nullptr || !FPaths::FileExists(*TextureFilePath)))
+					FallbackValue = pxr::GfVec4f{0.0, 0.0, 0.0, 1.0f};
+					ColorSpaceToken = UnrealIdentifiers::SRGBColorSpaceToken;
+					break;
+				case MP_Specular:
+					InputToken = UnrealIdentifiers::Specular;
+					InputType = pxr::SdfValueTypeNames->Float;
+					ConstantValue = UserConstantValue ? *UserConstantValue : ConstantLinearValue[0];
+					FallbackValue = pxr::GfVec4f{0.5, 0.5, 0.5, 1.0f};
+					break;
+				case MP_Metallic:
+					InputToken = UnrealIdentifiers::Metallic;
+					InputType = pxr::SdfValueTypeNames->Float;
+					ConstantValue = UserConstantValue ? *UserConstantValue : ConstantLinearValue[0];
+					FallbackValue = pxr::GfVec4f{0.0, 0.0, 0.0, 1.0f};
+					break;
+				case MP_Roughness:
+					InputToken = UnrealIdentifiers::Roughness;
+					InputType = pxr::SdfValueTypeNames->Float;
+					ConstantValue = UserConstantValue ? *UserConstantValue : ConstantLinearValue[0];
+					FallbackValue = pxr::GfVec4f{0.5, 0.5, 0.5, 1.0f};
+					break;
+				case MP_Normal:
+					InputToken = UnrealIdentifiers::Normal;
+					InputType = pxr::SdfValueTypeNames->Normal3f;
+					if (bPropertyValueIsConstant)
 					{
-						UE_LOG(
-							LogUsd,
-							Warning,
-							TEXT("Skipping material property %d as target texture '%s' was not found."),
-							Property,
-							TextureFilePath ? **TextureFilePath : TEXT("")
-						);
-						continue;
+						// This doesn't make much sense here but it's an available option, so here we go
+						if (UserConstantValue)
+						{
+							ConstantValue = pxr::GfVec3f(*UserConstantValue, *UserConstantValue, *UserConstantValue);
+						}
+						else
+						{
+							ConstantValue = ConstantLinearValue;
+						}
 					}
-
-					pxr::TfToken InputToken;
-					pxr::SdfValueTypeName InputType;
-					pxr::VtValue ConstantValue;
-					pxr::GfVec4f FallbackValue;
-					pxr::TfToken ColorSpaceToken = UnrealIdentifiers::RawColorSpaceToken;
-					switch (Property)
+					FallbackValue = pxr::GfVec4f{0.0, 0.0, 1.0, 1.0f};
+					break;
+				case MP_Tangent:
+					InputToken = UnrealIdentifiers::Tangent;
+					InputType = pxr::SdfValueTypeNames->Normal3f;
+					if (bPropertyValueIsConstant)
 					{
-						case MP_BaseColor:
-							InputToken = UnrealIdentifiers::DiffuseColor;
-							InputType = pxr::SdfValueTypeNames->Color3f;
-							if (bPropertyValueIsConstant)
-							{
-								if (UserConstantValue)
-								{
-									ConstantValue = pxr::GfVec3f(*UserConstantValue, *UserConstantValue, *UserConstantValue);
-								}
-								else
-								{
-									ConstantValue = ConstantLinearValue;
-								}
-							}
-							FallbackValue = pxr::GfVec4f{0.0, 0.0, 0.0, 1.0f};
-							ColorSpaceToken = UnrealIdentifiers::SRGBColorSpaceToken;
-							break;
-						case MP_Specular:
-							InputToken = UnrealIdentifiers::Specular;
-							InputType = pxr::SdfValueTypeNames->Float;
-							ConstantValue = UserConstantValue ? *UserConstantValue : ConstantLinearValue[0];
-							FallbackValue = pxr::GfVec4f{0.5, 0.5, 0.5, 1.0f};
-							break;
-						case MP_Metallic:
-							InputToken = UnrealIdentifiers::Metallic;
-							InputType = pxr::SdfValueTypeNames->Float;
-							ConstantValue = UserConstantValue ? *UserConstantValue : ConstantLinearValue[0];
-							FallbackValue = pxr::GfVec4f{0.0, 0.0, 0.0, 1.0f};
-							break;
-						case MP_Roughness:
-							InputToken = UnrealIdentifiers::Roughness;
-							InputType = pxr::SdfValueTypeNames->Float;
-							ConstantValue = UserConstantValue ? *UserConstantValue : ConstantLinearValue[0];
-							FallbackValue = pxr::GfVec4f{0.5, 0.5, 0.5, 1.0f};
-							break;
-						case MP_Normal:
-							InputToken = UnrealIdentifiers::Normal;
-							InputType = pxr::SdfValueTypeNames->Normal3f;
-							if (bPropertyValueIsConstant)
-							{
-								// This doesn't make much sense here but it's an available option, so here we go
-								if (UserConstantValue)
-								{
-									ConstantValue = pxr::GfVec3f(*UserConstantValue, *UserConstantValue, *UserConstantValue);
-								}
-								else
-								{
-									ConstantValue = ConstantLinearValue;
-								}
-							}
-							FallbackValue = pxr::GfVec4f{0.0, 0.0, 1.0, 1.0f};
-							break;
-						case MP_Tangent:
-							InputToken = UnrealIdentifiers::Tangent;
-							InputType = pxr::SdfValueTypeNames->Normal3f;
-							if (bPropertyValueIsConstant)
-							{
-								// This doesn't make much sense here but it's an available option, so here we go
-								if (UserConstantValue)
-								{
-									ConstantValue = pxr::GfVec3f(*UserConstantValue, *UserConstantValue, *UserConstantValue);
-								}
-								else
-								{
-									ConstantValue = ConstantLinearValue;
-								}
-							}
-							FallbackValue = pxr::GfVec4f{1.0, 0.0, 0.0, 1.0f};
-							break;
-						case MP_EmissiveColor:
-							InputToken = UnrealIdentifiers::EmissiveColor;
-							InputType = pxr::SdfValueTypeNames->Color3f;
-							if (bPropertyValueIsConstant)
-							{
-								// This doesn't make much sense here but it's an available option, so here we go
-								if (UserConstantValue)
-								{
-									ConstantValue = pxr::GfVec3f(*UserConstantValue, *UserConstantValue, *UserConstantValue);
-								}
-								else
-								{
-									ConstantValue = ConstantLinearValue;
-								}
-							}
-							FallbackValue = pxr::GfVec4f{0.0, 0.0, 0.0, 1.0f};
-							// Emissive is also written out with RawColorSpaceToken as we write them as EXR files now
-							break;
-						case MP_Opacity:	// It's OK that we use the same for both as these are mutually exclusive blend modes
-						case MP_OpacityMask:
-							InputToken = UnrealIdentifiers::Opacity;
-							InputType = pxr::SdfValueTypeNames->Float;
-							ConstantValue = UserConstantValue ? *UserConstantValue : ConstantLinearValue[0];
-							FallbackValue = pxr::GfVec4f{1.0, 1.0, 1.0, 1.0f};
-							break;
-						case MP_Anisotropy:
-							InputToken = UnrealIdentifiers::Anisotropy;
-							InputType = pxr::SdfValueTypeNames->Float;
-							ConstantValue = UserConstantValue ? *UserConstantValue : ConstantLinearValue[0];
-							FallbackValue = pxr::GfVec4f{0.0, 0.0, 0.0, 1.0f};
-							break;
-						case MP_AmbientOcclusion:
-							InputToken = UnrealIdentifiers::Occlusion;
-							InputType = pxr::SdfValueTypeNames->Float;
-							ConstantValue = UserConstantValue ? *UserConstantValue : ConstantLinearValue[0];
-							FallbackValue = pxr::GfVec4f{1.0, 1.0, 1.0, 1.0f};
-							break;
-						case MP_SubsurfaceColor:
-							InputToken = UnrealIdentifiers::SubsurfaceColor;
-							InputType = pxr::SdfValueTypeNames->Color3f;
-							if (bPropertyValueIsConstant)
-							{
-								if (UserConstantValue)
-								{
-									ConstantValue = pxr::GfVec3f(*UserConstantValue, *UserConstantValue, *UserConstantValue);
-								}
-								else
-								{
-									ConstantValue = ConstantLinearValue;
-								}
-							}
-							FallbackValue = pxr::GfVec4f{1.0, 1.0, 1.0, 1.0f};
-							ColorSpaceToken = UnrealIdentifiers::SRGBColorSpaceToken;
-							break;
-						default:
-							continue;
-							break;
+						// This doesn't make much sense here but it's an available option, so here we go
+						if (UserConstantValue)
+						{
+							ConstantValue = pxr::GfVec3f(*UserConstantValue, *UserConstantValue, *UserConstantValue);
+						}
+						else
+						{
+							ConstantValue = ConstantLinearValue;
+						}
 					}
-
-					pxr::UsdShadeInput ShadeInput = ShadeShader.CreateInput(InputToken, InputType);
+					FallbackValue = pxr::GfVec4f{1.0, 0.0, 0.0, 1.0f};
+					break;
+				case MP_EmissiveColor:
+					InputToken = UnrealIdentifiers::EmissiveColor;
+					InputType = pxr::SdfValueTypeNames->Color3f;
 					if (bPropertyValueIsConstant)
 					{
-						ShadeInput.Set(ConstantValue);
+						// This doesn't make much sense here but it's an available option, so here we go
+						if (UserConstantValue)
+						{
+							ConstantValue = pxr::GfVec3f(*UserConstantValue, *UserConstantValue, *UserConstantValue);
+						}
+						else
+						{
+							ConstantValue = ConstantLinearValue;
+						}
 					}
-					else	// Its a texture
+					FallbackValue = pxr::GfVec4f{0.0, 0.0, 0.0, 1.0f};
+					// Emissive is also written out with RawColorSpaceToken as we write them as EXR files now
+					break;
+				case MP_Opacity:	// It's OK that we use the same for both as these are mutually exclusive blend modes
+				case MP_OpacityMask:
+					InputToken = UnrealIdentifiers::Opacity;
+					InputType = pxr::SdfValueTypeNames->Float;
+					ConstantValue = UserConstantValue ? *UserConstantValue : ConstantLinearValue[0];
+					FallbackValue = pxr::GfVec4f{1.0, 1.0, 1.0, 1.0f};
+					break;
+				case MP_Anisotropy:
+					InputToken = UnrealIdentifiers::Anisotropy;
+					InputType = pxr::SdfValueTypeNames->Float;
+					ConstantValue = UserConstantValue ? *UserConstantValue : ConstantLinearValue[0];
+					FallbackValue = pxr::GfVec4f{0.0, 0.0, 0.0, 1.0f};
+					break;
+				case MP_AmbientOcclusion:
+					InputToken = UnrealIdentifiers::Occlusion;
+					InputType = pxr::SdfValueTypeNames->Float;
+					ConstantValue = UserConstantValue ? *UserConstantValue : ConstantLinearValue[0];
+					FallbackValue = pxr::GfVec4f{1.0, 1.0, 1.0, 1.0f};
+					break;
+				case MP_SubsurfaceColor:
+					InputToken = UnrealIdentifiers::SubsurfaceColor;
+					InputType = pxr::SdfValueTypeNames->Color3f;
+					if (bPropertyValueIsConstant)
 					{
-						// Create the primvar/uv set reader on-demand. We'll be using the same UV set for everything for now
-						if (!PrimvarReaderShader)
+						if (UserConstantValue)
 						{
-							PrimvarReaderShader = pxr::UsdShadeShader::Define(Stage, MaterialPath.AppendChild(PrimvarReaderShaderName));
-							PrimvarReaderShader.SetShaderId(UnrealIdentifiers::UsdPrimvarReader_float2);
+							ConstantValue = pxr::GfVec3f(*UserConstantValue, *UserConstantValue, *UserConstantValue);
+						}
+						else
+						{
+							ConstantValue = ConstantLinearValue;
+						}
+					}
+					FallbackValue = pxr::GfVec4f{1.0, 1.0, 1.0, 1.0f};
+					ColorSpaceToken = UnrealIdentifiers::SRGBColorSpaceToken;
+					break;
+				default:
+					continue;
+					break;
+			}
+
+			pxr::UsdShadeInput ShadeInput = ShadeShader.CreateInput(InputToken, InputType);
+			if (bPropertyValueIsConstant)
+			{
+				ShadeInput.Set(ConstantValue);
+			}
+			else	// Its a texture
+			{
+				// Create the primvar/uv set reader on-demand. We'll be using the same UV set for everything for now
+				if (!PrimvarReaderShader)
+				{
+					PrimvarReaderShader = pxr::UsdShadeShader::Define(Stage, MaterialPath.AppendChild(PrimvarReaderShaderName));
+					PrimvarReaderShader.SetShaderId(UnrealIdentifiers::UsdPrimvarReader_float2);
 
-							// Create the 'st' input directly on the material, as that seems to be preferred
-							pxr::UsdShadeInput MaterialStInput = OutUsdShadeMaterial.CreateInput(PrimvarVariableName, pxr::SdfValueTypeNames->Token);
-							MaterialStInput.Set(UnrealIdentifiers::St);
+					// Create the 'st' input directly on the material, as that seems to be preferred
+					pxr::UsdShadeInput MaterialStInput = OutUsdShadeMaterial.CreateInput(PrimvarVariableName, pxr::SdfValueTypeNames->Token);
+					MaterialStInput.Set(UnrealIdentifiers::St);
 
-							pxr::UsdShadeInput VarnameInput = PrimvarReaderShader.CreateInput(
-								UnrealIdentifiers::Varname,
-								pxr::SdfValueTypeNames->String
-							);
-							VarnameInput.ConnectToSource(MaterialStInput);
+					pxr::UsdShadeInput VarnameInput = PrimvarReaderShader.CreateInput(UnrealIdentifiers::Varname, pxr::SdfValueTypeNames->String);
+					VarnameInput.ConnectToSource(MaterialStInput);
 
-							PrimvarReaderShader.CreateOutput(UnrealIdentifiers::Result, pxr::SdfValueTypeNames->Token);
-						}
+					PrimvarReaderShader.CreateOutput(UnrealIdentifiers::Result, pxr::SdfValueTypeNames->Token);
+				}
 
-						FString TextureReaderName = UsdToUnreal::ConvertToken(InputToken);
-						TextureReaderName.RemoveFromEnd(TEXT("Color"));
-						TextureReaderName += TEXT("Texture");
+				FString TextureReaderName = UsdToUnreal::ConvertToken(InputToken);
+				TextureReaderName.RemoveFromEnd(TEXT("Color"));
+				TextureReaderName += TEXT("Texture");
 
-						pxr::UsdShadeShader UsdUVTextureShader = pxr::UsdShadeShader::Define(
-							Stage,
-							MaterialPath.AppendChild(UnrealToUsd::ConvertToken(*TextureReaderName).Get())
-						);
-						UsdUVTextureShader.SetShaderId(UnrealIdentifiers::UsdUVTexture);
+				pxr::UsdShadeShader UsdUVTextureShader = pxr::UsdShadeShader::Define(
+					Stage,
+					MaterialPath.AppendChild(UnrealToUsd::ConvertToken(*TextureReaderName).Get())
+				);
+				UsdUVTextureShader.SetShaderId(UnrealIdentifiers::UsdUVTexture);
 
-						pxr::UsdShadeInput TextureFileInput = UsdUVTextureShader.CreateInput(UnrealIdentifiers::File, pxr::SdfValueTypeNames->Asset);
-						FString TextureRelativePath = *TextureFilePath;
-						if (!UsdFilePath.IsEmpty())
-						{
-							FPaths::MakePathRelativeTo(TextureRelativePath, *UsdFilePath);
-						}
-						TextureFileInput.Set(pxr::SdfAssetPath(UnrealToUsd::ConvertString(*TextureRelativePath).Get()));
-
-						pxr::UsdShadeInput TextureStInput = UsdUVTextureShader.CreateInput(UnrealIdentifiers::St, pxr::SdfValueTypeNames->Float2);
-						TextureStInput.ConnectToSource(PrimvarReaderShader.GetOutput(UnrealIdentifiers::Result));
-
-						pxr::UsdShadeInput TextureColorSpaceInput = UsdUVTextureShader.CreateInput(
-							UnrealIdentifiers::SourceColorSpaceToken,
-							pxr::SdfValueTypeNames->Token
-						);
-						TextureColorSpaceInput.Set(ColorSpaceToken);
-
-						pxr::UsdShadeInput TextureFallbackInput = UsdUVTextureShader.CreateInput(
-							UnrealIdentifiers::Fallback,
-							pxr::SdfValueTypeNames->Float4
-						);
-						TextureFallbackInput.Set(FallbackValue);
-
-						// In the general case it's impossible to set a "correct" wrapping value here because the
-						// material we just baked may be using 3 different textures with UV transforms an all different
-						// texture wrapping modes, and we're forced to pick a single value to wrap the baked texture
-						// with, but let's at least write "repeat" out as that is the default for textures in UE and
-						// the more likely to be correct, in case the mesh does things like have UVs outside [0, 1]
-						pxr::UsdShadeInput TextureFileWrapSInput = UsdUVTextureShader.CreateInput(
-							UnrealIdentifiers::WrapS,
-							pxr::SdfValueTypeNames->Token
-						);
-						TextureFileWrapSInput.Set(UnrealIdentifiers::Repeat);
-						pxr::UsdShadeInput TextureFileWrapTInput = UsdUVTextureShader.CreateInput(
-							UnrealIdentifiers::WrapT,
-							pxr::SdfValueTypeNames->Token
-						);
-						TextureFileWrapTInput.Set(UnrealIdentifiers::Repeat);
-
-						pxr::UsdShadeOutput TextureOutput = UsdUVTextureShader.CreateOutput(
-							InputType == pxr::SdfValueTypeNames->Float ? UnrealIdentifiers::R : UnrealIdentifiers::RGB,
-							InputType
-						);
-
-						ShadeInput.ConnectToSource(TextureOutput);
-					}
+				pxr::UsdShadeInput TextureFileInput = UsdUVTextureShader.CreateInput(UnrealIdentifiers::File, pxr::SdfValueTypeNames->Asset);
+				FString TextureRelativePath = *TextureFilePath;
+				if (!UsdFilePath.IsEmpty())
+				{
+					FPaths::MakePathRelativeTo(TextureRelativePath, *UsdFilePath);
 				}
+				TextureFileInput.Set(pxr::SdfAssetPath(UnrealToUsd::ConvertString(*TextureRelativePath).Get()));
 
-				return true;
+				pxr::UsdShadeInput TextureStInput = UsdUVTextureShader.CreateInput(UnrealIdentifiers::St, pxr::SdfValueTypeNames->Float2);
+				TextureStInput.ConnectToSource(PrimvarReaderShader.GetOutput(UnrealIdentifiers::Result));
+
+				pxr::UsdShadeInput TextureColorSpaceInput = UsdUVTextureShader.CreateInput(
+					UnrealIdentifiers::SourceColorSpaceToken,
+					pxr::SdfValueTypeNames->Token
+				);
+				TextureColorSpaceInput.Set(ColorSpaceToken);
+
+				pxr::UsdShadeInput TextureFallbackInput = UsdUVTextureShader.CreateInput(UnrealIdentifiers::Fallback, pxr::SdfValueTypeNames->Float4);
+				TextureFallbackInput.Set(FallbackValue);
+
+				// In the general case it's impossible to set a "correct" wrapping value here because the
+				// material we just baked may be using 3 different textures with UV transforms an all different
+				// texture wrapping modes, and we're forced to pick a single value to wrap the baked texture
+				// with, but let's at least write "repeat" out as that is the default for textures in UE and
+				// the more likely to be correct, in case the mesh does things like have UVs outside [0, 1]
+				pxr::UsdShadeInput TextureFileWrapSInput = UsdUVTextureShader.CreateInput(UnrealIdentifiers::WrapS, pxr::SdfValueTypeNames->Token);
+				TextureFileWrapSInput.Set(UnrealIdentifiers::Repeat);
+				pxr::UsdShadeInput TextureFileWrapTInput = UsdUVTextureShader.CreateInput(UnrealIdentifiers::WrapT, pxr::SdfValueTypeNames->Token);
+				TextureFileWrapTInput.Set(UnrealIdentifiers::Repeat);
+
+				pxr::UsdShadeOutput TextureOutput = UsdUVTextureShader.CreateOutput(
+					InputType == pxr::SdfValueTypeNames->Float ? UnrealIdentifiers::R : UnrealIdentifiers::RGB,
+					InputType
+				);
+
+				ShadeInput.ConnectToSource(TextureOutput);
 			}
+		}
+
+		return true;
+	}
 #endif	  // WITH_EDITOR
 
-			void HashShadeInput(const pxr::UsdShadeInput& ShadeInput, FSHA1& InOutHashState)
+	void HashShadeInput(const pxr::UsdShadeInput& ShadeInput, FSHA1& InOutHashState)
+	{
+		if (!ShadeInput)
+		{
+			return;
+		}
+
+		FScopedUsdAllocs UsdAllocs;
+
+		FString InputName = UsdToUnreal::ConvertToken(ShadeInput.GetBaseName());
+		InOutHashState.UpdateWithString(*InputName, InputName.Len());
+
+		FString InputTypeName = UsdToUnreal::ConvertToken(ShadeInput.GetTypeName().GetAsToken());
+		InOutHashState.UpdateWithString(*InputTypeName, InputTypeName.Len());
+
+		// Connected to something else, recurse
+		pxr::UsdShadeConnectableAPI Source;
+		pxr::TfToken SourceName;
+		pxr::UsdShadeAttributeType AttributeType;
+		if (pxr::UsdShadeConnectableAPI::GetConnectedSource(ShadeInput.GetAttr(), &Source, &SourceName, &AttributeType))
+		{
+			FString SourceOutputName = UsdToUnreal::ConvertToken(SourceName);
+			InOutHashState.UpdateWithString(*SourceOutputName, SourceOutputName.Len());
+
+			// Skip in case our input is connected to an output on the same prim, or else we'll recurse forever
+			if (Source.GetPrim() == ShadeInput.GetPrim())
 			{
-				if (!ShadeInput)
-				{
-					return;
-				}
+				return;
+			}
+
+			for (const pxr::UsdShadeInput& ChildInput : Source.GetInputs())
+			{
+				HashShadeInput(ChildInput, InOutHashState);
+			}
+		}
+		// Not connected to anything, has a value (this could be a texture file path too)
+		else
+		{
+			pxr::VtValue ShadeInputValue;
+			ShadeInput.Get(&ShadeInputValue);
+
+			// We have to manually resolve and hash these file paths or else resolvedpaths inside usdz archives will have upper case driver
+			// letters when we first open the stage, but will switch to lower case drive letters if we reload them. This is not something
+			// we're doing, as it happens with a pure USD python script. (this with USD 21.05 in June 2021)
+			if (ShadeInputValue.IsHolding<pxr::SdfAssetPath>())
+			{
+				FString ResolvedPath = UsdUtils::GetResolvedAssetPath(ShadeInput.GetAttr());
+				InOutHashState.UpdateWithString(*ResolvedPath, ResolvedPath.Len());
+			}
+			else if (ShadeInputValue.IsHolding<pxr::TfToken>())
+			{
+				// Stringify instead of using GetHash because if ShadeInputValue contains a pxr::TfToken then
+				// it will actually just contain some non-deterministic integer IDs
+				FString Stringified = UsdToUnreal::ConvertString(pxr::TfStringify(ShadeInputValue));
+				InOutHashState.UpdateWithString(*Stringified, Stringified.Len());
+			}
+			else
+			{
+				uint64 ValueHash = (uint64)ShadeInputValue.GetHash();
+				InOutHashState.Update(reinterpret_cast<uint8*>(&ValueHash), sizeof(uint64));
+			}
+		}
+	}
 
-				FScopedUsdAllocs UsdAllocs;
+	// Fills out the Texture parameters within InOutMaterialData, spawning new UTexture2Ds within
+	// the TexturesCache and updating the FTextureParameterValue::Texture member to point at them
+	void ResolveTextureParameters(
+		const pxr::UsdPrim& MaterialPrim,
+		UMaterialInterface& Material,
+		UUsdAssetCache3* TexturesCache,
+		UsdToUnreal::FUsdPreviewSurfaceMaterialData& InOutMaterialData,
+		bool bShareAssetsForIdenticalPrims
+	)
+	{
+		if (!TexturesCache)
+		{
+			return;
+		}
 
-				FString InputName = UsdToUnreal::ConvertToken(ShadeInput.GetBaseName());
-				InOutHashState.UpdateWithString(*InputName, InputName.Len());
+		UUsdAssetUserData* MaterialAssetUserData = UsdUnreal::ObjectUtils::GetAssetUserData(&Material);
 
-				FString InputTypeName = UsdToUnreal::ConvertToken(ShadeInput.GetTypeName().GetAsToken());
-				InOutHashState.UpdateWithString(*InputTypeName, InputTypeName.Len());
+		for (TPair<FString, UsdToUnreal::FParameterValue>& Pair : InOutMaterialData.Parameters)
+		{
+			UsdToUnreal::FTextureParameterValue* TextureParameter = Pair.Value.TryGet<UsdToUnreal::FTextureParameterValue>();
+			if (!TextureParameter)
+			{
+				continue;
+			}
 
-				// Connected to something else, recurse
-				pxr::UsdShadeConnectableAPI Source;
-				pxr::TfToken SourceName;
-				pxr::UsdShadeAttributeType AttributeType;
-				if (pxr::UsdShadeConnectableAPI::GetConnectedSource(ShadeInput.GetAttr(), &Source, &SourceName, &AttributeType))
-				{
-					FString SourceOutputName = UsdToUnreal::ConvertToken(SourceName);
-					InOutHashState.UpdateWithString(*SourceOutputName, SourceOutputName.Len());
+			// We'll add these to the hash because the materials are built to try and reuse the same textures for multiple channels,
+			// and those may expect linear or sRGB values. Without this we may parse a texture as linear because we hit the opacity
+			// channel first, and then reuse it as linear for the base color channel even though it should have been sRGB. Plus we may
+			// have something weird like a normal map being plugged into the base color and the normal channel.
+			TextureCompressionSettings CompressionSettings = TextureParameter->Group == TEXTUREGROUP_WorldNormalMap
+																 ? TextureCompressionSettings::TC_Normalmap
+																 : TextureCompressionSettings::TC_Default;
+			const bool bSRGB = TextureParameter->GetSRGBValue();
+
+			// Note: We used to use the ShadeInput itself here on the asset hash prefix, which is now a bit awkward to get after the
+			// refactor for Interchange. Using the material prim itself should hopefully still be fine though, as it will provide a
+			// separate texture for each Material prim itself
+			const FString PrefixedTextureHash = UsdUtils::GetAssetHashPrefix(MaterialPrim, bShareAssetsForIdenticalPrims)
+												+ UsdUtils::GetTextureHash(	   //
+													TextureParameter->TextureFilePath,
+													bSRGB,
+													CompressionSettings,
+													TextureParameter->AddressX,
+													TextureParameter->AddressY
+												);
+
+			const FString& DesiredTextureName = FPaths::GetBaseFilename(TextureParameter->TextureFilePath);
+
+			const EObjectFlags DesiredFlags = RF_Public | RF_Standalone | RF_Transactional;
+
+			bool bCreatedTexture = false;
+			UTexture* Texture = TexturesCache->GetOrCreateCustomCachedAsset<UTexture2D>(
+				PrefixedTextureHash,
+				DesiredTextureName,
+				DesiredFlags,
+				[TextureParameter, TexturesCache, bSRGB](UPackage* Outer, FName SanitizedName, EObjectFlags DesiredFlags)
+				{
+					return UsdUtils::CreateTexture(
+						TextureParameter->TextureFilePath,
+						SanitizedName,
+						TextureParameter->Group,
+						DesiredFlags,
+						Outer,
+						!bSRGB
+					);
+				},
+				&bCreatedTexture
+			);
 
-					// Skip in case our input is connected to an output on the same prim, or else we'll recurse forever
-					if (Source.GetPrim() == ShadeInput.GetPrim())
-					{
-						return;
-					}
+			if (Texture)
+			{
+				TextureParameter->Texture = Texture;
 
-					for (const pxr::UsdShadeInput& ChildInput : Source.GetInputs())
+				// The texture resource needs to be updated only if the following settings have changed from their default values
+				bool bNeedUpdateResource = false;
+				bNeedUpdateResource |= TextureParameter->AddressX != TextureAddress::TA_Wrap;
+				bNeedUpdateResource |= TextureParameter->AddressY != TextureAddress::TA_Wrap;
+#if !WITH_EDITOR
+				// In editor, these settings are already set on the factory when the texture is created so no further update is needed
+				bNeedUpdateResource |= bSRGB != true;
+				// TC_Normalmap is set via the TEXTUREGROUP_WorldNormalMap Group
+				bNeedUpdateResource |= CompressionSettings != TextureCompressionSettings::TC_Default;
+#endif	  // !WITH_EDITOR
+				if (bCreatedTexture && bNeedUpdateResource)
+				{
+#if WITH_EDITOR
+					Texture->PreEditChange(nullptr);
+#endif	  // WITH_EDITOR
+					Texture->SRGB = bSRGB;
+					Texture->CompressionSettings = CompressionSettings;
+					if (UTexture2D* Texture2D = Cast<UTexture2D>(Texture))
 					{
-						HashShadeInput(ChildInput, InOutHashState);
+						Texture2D->AddressX = TextureParameter->AddressX;
+						Texture2D->AddressY = TextureParameter->AddressY;
 					}
+#if WITH_EDITOR
+					Texture->PostEditChange();
+#else
+					Texture->UpdateResource();
+#endif	  // WITH_EDITOR
 				}
-				// Not connected to anything, has a value (this could be a texture file path too)
-				else
-				{
-					pxr::VtValue ShadeInputValue;
-					ShadeInput.Get(&ShadeInputValue);
 
-					// We have to manually resolve and hash these file paths or else resolvedpaths inside usdz archives will have upper case driver
-					// letters when we first open the stage, but will switch to lower case drive letters if we reload them. This is not something
-					// we're doing, as it happens with a pure USD python script. (this with USD 21.05 in June 2021)
-					if (ShadeInputValue.IsHolding<pxr::SdfAssetPath>())
-					{
-						FString ResolvedPath = UsdUtils::GetResolvedAssetPath(ShadeInput.GetAttr());
-						InOutHashState.UpdateWithString(*ResolvedPath, ResolvedPath.Len());
-					}
-					else if (ShadeInputValue.IsHolding<pxr::TfToken>())
-					{
-						// Stringify instead of using GetHash because if ShadeInputValue contains a pxr::TfToken then
-						// it will actually just contain some non-deterministic integer IDs
-						FString Stringified = UsdToUnreal::ConvertString(pxr::TfStringify(ShadeInputValue));
-						InOutHashState.UpdateWithString(*Stringified, Stringified.Len());
-					}
-					else
+				// Give the same prim paths to the texture, so that it ends up imported right next to the material
+				if (MaterialAssetUserData)
+				{
+					if (UUsdAssetUserData* AssetUserData = UsdUnreal::ObjectUtils::GetOrCreateAssetUserData(Texture))
 					{
-						uint64 ValueHash = (uint64)ShadeInputValue.GetHash();
-						InOutHashState.Update(reinterpret_cast<uint8*>(&ValueHash), sizeof(uint64));
+						for (const FString& MaterialPrimPath : MaterialAssetUserData->PrimPaths)
+						{
+							AssetUserData->PrimPaths.AddUnique(MaterialPrimPath);
+						}
 					}
 				}
 			}
-		}	 // namespace Private
-	}		 // namespace UsdShadeConversion
-}	 // namespace UE
+			else
+			{
+				FUsdLogManager::LogMessage(
+					EMessageSeverity::Warning,
+					FText::Format(
+						LOCTEXT("FailedToParseTexture", "Failed to parse texture at path '{0}'"),
+						FText::FromString(TextureParameter->TextureFilePath)
+					)
+				);
+			}
+		}
+	}
+}	 // namespace UE::UsdShadeConversion::Private
 namespace UsdShadeConversionImpl = UE::UsdShadeConversion::Private;
 
+bool UsdToUnreal::FTextureParameterValue::GetSRGBValue() const
+{
+	bool bFallbackSRGB = true;
+	// Disable SRGB when parsing float textures, as they're likely specular/roughness maps
+	if (Group == TEXTUREGROUP_WorldNormalMap || Group == TEXTUREGROUP_WorldSpecular)
+	{
+		bFallbackSRGB = false;
+	}
+	return bSRGB.Get(bFallbackSRGB);
+}
+
 bool UsdToUnreal::ConvertMaterial(
-	const pxr::UsdShadeMaterial& UsdShadeMaterial,
-	UMaterialInstance& Material,
-	UUsdAssetCache3* TexturesCache,
-	const TCHAR* RenderContext,
-	bool bShareAssetsForIdenticalPrims
+	const pxr::UsdPrim& UsdShadeMaterialPrim,
+	FUsdPreviewSurfaceMaterialData& OutMaterialData,
+	const TCHAR* RenderContext
 )
 {
+	using namespace UsdShadeConversionImpl;
+
 	FScopedUsdAllocs UsdAllocs;
 
+	pxr::UsdShadeMaterial UsdShadeMaterial{UsdShadeMaterialPrim};
+	if (!UsdShadeMaterial)
+	{
+		return false;
+	}
+
 	pxr::TfToken RenderContextToken = pxr::UsdShadeTokens->universalRenderContext;
 	if (RenderContext)
 	{
@@ -2176,136 +2167,64 @@ bool UsdToUnreal::ConvertMaterial(
 
 	pxr::UsdShadeConnectableAPI Connectable{SurfaceShader};
 
-	UsdShadeConversionImpl::FParameterValue ParameterValue;
-
-	TMap<FString, UsdShadeConversionImpl::FParameterValue> MaterialParameters;
-	MaterialParameters.Reserve(8);
-
-	const bool bIsNormalMap = true;
+	UsdToUnreal::FParameterValue ParameterValue;
 
 	// Base color
-	if (UsdShadeConversionImpl::GetVec3ParameterValue(
-			Connectable,
-			UnrealIdentifiers::DiffuseColor,
-			FLinearColor(0, 0, 0),
-			ParameterValue,
-			!bIsNormalMap,
-			&Material,
-			TexturesCache,
-			bShareAssetsForIdenticalPrims
-		))
+	if (GetVec3ParameterValue(Connectable, UnrealIdentifiers::DiffuseColor, FLinearColor::Black, ParameterValue))
 	{
-		MaterialParameters.Add(TEXT("BaseColor"), ParameterValue);
+		OutMaterialData.Parameters.Add(TEXT("BaseColor"), ParameterValue);
 	}
 
 	// Emissive color
-	if (UsdShadeConversionImpl::GetVec3ParameterValue(
-			Connectable,
-			UnrealIdentifiers::EmissiveColor,
-			FLinearColor(0, 0, 0),
-			ParameterValue,
-			!bIsNormalMap,
-			&Material,
-			TexturesCache,
-			bShareAssetsForIdenticalPrims
-		))
+	if (GetVec3ParameterValue(Connectable, UnrealIdentifiers::EmissiveColor, FLinearColor::Black, ParameterValue))
 	{
-		MaterialParameters.Add(TEXT("EmissiveColor"), ParameterValue);
+		OutMaterialData.Parameters.Add(TEXT("EmissiveColor"), ParameterValue);
 	}
 
 	// Metallic
-	if (UsdShadeConversionImpl::GetFloatParameterValue(
-			Connectable,
-			UnrealIdentifiers::Metallic,
-			0.f,
-			ParameterValue,
-			&Material,
-			TexturesCache,
-			bShareAssetsForIdenticalPrims
-		))
+	if (GetFloatParameterValue(Connectable, UnrealIdentifiers::Metallic, 0.f, ParameterValue))
 	{
-		MaterialParameters.Add(TEXT("Metallic"), ParameterValue);
+		OutMaterialData.Parameters.Add(TEXT("Metallic"), ParameterValue);
 	}
 
 	// Roughness
-	if (UsdShadeConversionImpl::GetFloatParameterValue(
-			Connectable,
-			UnrealIdentifiers::Roughness,
-			1.f,
-			ParameterValue,
-			&Material,
-			TexturesCache,
-			bShareAssetsForIdenticalPrims
-		))
+	if (GetFloatParameterValue(Connectable, UnrealIdentifiers::Roughness, 1.f, ParameterValue))
 	{
-		MaterialParameters.Add(TEXT("Roughness"), ParameterValue);
+		OutMaterialData.Parameters.Add(TEXT("Roughness"), ParameterValue);
 	}
 
 	// Opacity
-	if (UsdShadeConversionImpl::GetFloatParameterValue(
-			Connectable,
-			UnrealIdentifiers::Opacity,
-			1.f,
-			ParameterValue,
-			&Material,
-			TexturesCache,
-			bShareAssetsForIdenticalPrims
-		))
+	if (GetFloatParameterValue(Connectable, UnrealIdentifiers::Opacity, 1.f, ParameterValue))
 	{
-		MaterialParameters.Add(TEXT("Opacity"), ParameterValue);
+		OutMaterialData.Parameters.Add(TEXT("Opacity"), ParameterValue);
 	}
 
 	// Normal
-	if (UsdShadeConversionImpl::GetVec3ParameterValue(
-			Connectable,
-			UnrealIdentifiers::Normal,
-			FLinearColor(0, 0, 1),
-			ParameterValue,
-			bIsNormalMap,
-			&Material,
-			TexturesCache,
-			bShareAssetsForIdenticalPrims
-		))
+	if (GetVec3ParameterValue(Connectable, UnrealIdentifiers::Normal, FLinearColor::Blue, ParameterValue, TEXTUREGROUP_WorldNormalMap))
 	{
-		MaterialParameters.Add(TEXT("Normal"), ParameterValue);
+		OutMaterialData.Parameters.Add(TEXT("Normal"), ParameterValue);
 	}
 
 	// Refraction
-	if (UsdShadeConversionImpl::GetFloatParameterValue(
-			Connectable,
-			UnrealIdentifiers::Refraction,
-			1.5f,
-			ParameterValue,
-			&Material,
-			TexturesCache,
-			bShareAssetsForIdenticalPrims
-		))
+	if (GetFloatParameterValue(Connectable, UnrealIdentifiers::Refraction, 1.5f, ParameterValue))
 	{
-		MaterialParameters.Add(TEXT("Refraction"), ParameterValue);
+		OutMaterialData.Parameters.Add(TEXT("Refraction"), ParameterValue);
 	}
 
 	// Occlusion
-	if (UsdShadeConversionImpl::GetFloatParameterValue(
-			Connectable,
-			UnrealIdentifiers::Occlusion,
-			1.0f,
-			ParameterValue,
-			&Material,
-			TexturesCache,
-			bShareAssetsForIdenticalPrims
-		))
+	if (GetFloatParameterValue(Connectable, UnrealIdentifiers::Occlusion, 1.0f, ParameterValue))
 	{
-		MaterialParameters.Add(TEXT("AmbientOcclusion"), ParameterValue);
+		OutMaterialData.Parameters.Add(TEXT("AmbientOcclusion"), ParameterValue);
 	}
 
 	TMap<FString, FString> ParameterToPrimvar;
+	TMap<FString, int32> PrimvarToUVIndex;
 
 	// Sort primvars
 	TSet<FString> UsedPrimvars;
-	for (const TPair<FString, UsdShadeConversionImpl::FParameterValue>& Parameter : MaterialParameters)
+	for (const TPair<FString, UsdToUnreal::FParameterValue>& Parameter : OutMaterialData.Parameters)
 	{
-		if (const UsdShadeConversionImpl::FTextureParameterValue*
-				TextureParameterValue = Parameter.Value.TryGet<UsdShadeConversionImpl::FTextureParameterValue>())
+		if (const UsdToUnreal::FTextureParameterValue* TextureParameterValue = Parameter.Value.TryGet<UsdToUnreal::FTextureParameterValue>())
 		{
 			UsedPrimvars.Add(TextureParameterValue->Primvar);
 			ParameterToPrimvar.Add(Parameter.Key, TextureParameterValue->Primvar);
@@ -2314,7 +2233,6 @@ bool UsdToUnreal::ConvertMaterial(
 	UsedPrimvars.Remove(TEXT(""));
 	TArray<FString> SortedPrimvars = UsedPrimvars.Array();
 	SortedPrimvars.Sort();	  // Try for some deterministic ordering (st0 should come before st1, etc.)
-	TMap<FString, int32> PrimvarToUVIndex;
 	PrimvarToUVIndex.Reserve(SortedPrimvars.Num());
 	int32 UVIndex = 0;
 	for (const FString& Primvar : SortedPrimvars)
@@ -2322,11 +2240,40 @@ bool UsdToUnreal::ConvertMaterial(
 		PrimvarToUVIndex.Add(Primvar, UVIndex++);
 	}
 
+	OutMaterialData.PrimvarToUVIndex = PrimvarToUVIndex;
+
+	return true;
+}
+
+bool UsdToUnreal::ConvertMaterial(
+	const pxr::UsdShadeMaterial& UsdShadeMaterial,
+	UMaterialInstance& Material,
+	UUsdAssetCache3* TexturesCache,
+	const TCHAR* RenderContext,
+	bool bShareAssetsForIdenticalPrims
+)
+{
+	pxr::UsdPrim Prim = UsdShadeMaterial.GetPrim();
+	if (!Prim)
+	{
+		return false;
+	}
+
+	FUsdPreviewSurfaceMaterialData MaterialData;
+	const bool bSuccess = ConvertMaterial(Prim, MaterialData, RenderContext);
+	if (!bSuccess)
+	{
+		return false;
+	}
+
+	// Actually create textures if we've been given a place where to put them.
+	UsdShadeConversionImpl::ResolveTextureParameters(Prim, Material, TexturesCache, MaterialData, bShareAssetsForIdenticalPrims);
+
 	// Set material parameters on the actual material instance
 	const bool bForUsdPreviewSurface = true;
-	for (const TPair<FString, UsdShadeConversionImpl::FParameterValue>& Parameter : MaterialParameters)
+	for (const TPair<FString, UsdToUnreal::FParameterValue>& Parameter : MaterialData.Parameters)
 	{
-		UsdShadeConversionImpl::SetParameterValue(Material, *Parameter.Key, Parameter.Value, bForUsdPreviewSurface, PrimvarToUVIndex);
+		UsdShadeConversionImpl::SetParameterValue(Material, *Parameter.Key, Parameter.Value, bForUsdPreviewSurface, MaterialData.PrimvarToUVIndex);
 	}
 
 	// Handle world space normals
@@ -2345,7 +2292,18 @@ bool UsdToUnreal::ConvertMaterial(
 	// values if we need to
 	if (UUsdMaterialAssetUserData* UserData = Material.GetAssetUserData<UUsdMaterialAssetUserData>())
 	{
-		UserData->PrimvarToUVIndex = PrimvarToUVIndex;
+		// We rebuild this as there's no reason to store this separately on MaterialData itself, as the complete
+		// parameter data is stored there now
+		TMap<FString, FString> ParameterToPrimvar;
+		for (const TPair<FString, UsdToUnreal::FParameterValue>& Parameter : MaterialData.Parameters)
+		{
+			if (const UsdToUnreal::FTextureParameterValue* TextureParameterValue = Parameter.Value.TryGet<UsdToUnreal::FTextureParameterValue>())
+			{
+				ParameterToPrimvar.Add(Parameter.Key, TextureParameterValue->Primvar);
+			}
+		}
+
+		UserData->PrimvarToUVIndex = MaterialData.PrimvarToUVIndex;
 		UserData->ParameterToPrimvar = ParameterToPrimvar;
 	}
 
@@ -2363,198 +2321,33 @@ bool UsdToUnreal::ConvertMaterial(
 )
 {
 #if WITH_EDITOR
-	pxr::TfToken RenderContextToken = pxr::UsdShadeTokens->universalRenderContext;
-	if (RenderContext)
-	{
-		pxr::TfToken ProvidedRenderContextToken = UnrealToUsd::ConvertToken(RenderContext).Get();
-
-		// We won't ever create an asset for the Unreal render context material prim.
-		// Check the universal context here in case this material should also generate an asset for that context too
-		if (ProvidedRenderContextToken != UnrealIdentifiers::Unreal)
-		{
-			RenderContextToken = ProvidedRenderContextToken;
-		}
-	}
 
-	pxr::UsdShadeShader SurfaceShader = UsdShadeMaterial.ComputeSurfaceSource(RenderContextToken);
-	if (!SurfaceShader)
+	pxr::UsdPrim Prim = UsdShadeMaterial.GetPrim();
+	if (!Prim)
 	{
 		return false;
 	}
 
-	pxr::UsdShadeConnectableAPI Connectable{SurfaceShader};
-
-	UsdShadeConversionImpl::FParameterValue TempParameterValue;
-	TOptional<UsdShadeConversionImpl::FParameterValue> BaseColorParameter;
-	TOptional<UsdShadeConversionImpl::FParameterValue> EmissiveParameter;
-	TOptional<UsdShadeConversionImpl::FParameterValue> MetallicParameter;
-	TOptional<UsdShadeConversionImpl::FParameterValue> RoughnessParameter;
-	TOptional<UsdShadeConversionImpl::FParameterValue> OpacityParameter;
-	TOptional<UsdShadeConversionImpl::FParameterValue> NormalParameter;
-	TOptional<UsdShadeConversionImpl::FParameterValue> RefractionParameter;
-	TOptional<UsdShadeConversionImpl::FParameterValue> AmbientOcclusionParameter;
-
-	const bool bIsNormalMap = true;
-
-	// Base color
-	if (UsdShadeConversionImpl::GetVec3ParameterValue(
-			Connectable,
-			UnrealIdentifiers::DiffuseColor,
-			FLinearColor(0, 0, 0),
-			TempParameterValue,
-			!bIsNormalMap,
-			&Material,
-			TexturesCache,
-			bShareAssetsForIdenticalPrims
-		))
-	{
-		BaseColorParameter = TempParameterValue;
-	}
-
-	// Emissive color
-	if (UsdShadeConversionImpl::GetVec3ParameterValue(
-			Connectable,
-			UnrealIdentifiers::EmissiveColor,
-			FLinearColor(0, 0, 0),
-			TempParameterValue,
-			!bIsNormalMap,
-			&Material,
-			TexturesCache,
-			bShareAssetsForIdenticalPrims
-		))
-	{
-		EmissiveParameter = TempParameterValue;
-	}
-
-	// Metallic
-	if (UsdShadeConversionImpl::GetFloatParameterValue(
-			Connectable,
-			UnrealIdentifiers::Metallic,
-			0.f,
-			TempParameterValue,
-			&Material,
-			TexturesCache,
-			bShareAssetsForIdenticalPrims
-		))
-	{
-		MetallicParameter = TempParameterValue;
-	}
-
-	// Roughness
-	if (UsdShadeConversionImpl::GetFloatParameterValue(
-			Connectable,
-			UnrealIdentifiers::Roughness,
-			1.f,
-			TempParameterValue,
-			&Material,
-			TexturesCache,
-			bShareAssetsForIdenticalPrims
-		))
-	{
-		RoughnessParameter = TempParameterValue;
-	}
-
-	// Opacity
-	if (UsdShadeConversionImpl::GetFloatParameterValue(
-			Connectable,
-			UnrealIdentifiers::Opacity,
-			1.f,
-			TempParameterValue,
-			&Material,
-			TexturesCache,
-			bShareAssetsForIdenticalPrims
-		))
-	{
-		OpacityParameter = TempParameterValue;
-	}
-
-	// Normal
-	if (UsdShadeConversionImpl::GetVec3ParameterValue(
-			Connectable,
-			UnrealIdentifiers::Normal,
-			FLinearColor(0, 0, 1),
-			TempParameterValue,
-			bIsNormalMap,
-			&Material,
-			TexturesCache,
-			bShareAssetsForIdenticalPrims
-		))
-	{
-		NormalParameter = TempParameterValue;
-	}
-
-	// Refraction
-	const bool bHasRefractionValue = UsdShadeConversionImpl::GetFloatParameterValue(
-		Connectable,
-		UnrealIdentifiers::Refraction,
-		1.5f,
-		TempParameterValue,
-		&Material,
-		TexturesCache,
-		bShareAssetsForIdenticalPrims
-	);
-	if (bHasRefractionValue || Material.BlendMode == BLEND_Translucent)	   // Force a 1.5 IOR if USD didn't specify a value, as it's USD's fallback
-																		   // value
-	{
-		RefractionParameter = TempParameterValue;
-	}
-
-	// Ambient occlusion
-	if (UsdShadeConversionImpl::GetFloatParameterValue(
-			Connectable,
-			UnrealIdentifiers::Occlusion,
-			1.f,
-			TempParameterValue,
-			&Material,
-			TexturesCache,
-			bShareAssetsForIdenticalPrims
-		))
+	FUsdPreviewSurfaceMaterialData MaterialData;
+	const bool bSuccess = ConvertMaterial(Prim, MaterialData, RenderContext);
+	if (!bSuccess)
 	{
-		AmbientOcclusionParameter = TempParameterValue;
+		return false;
 	}
 
-	TMap<FString, FString> ParameterToPrimvar;
+	// Actually create textures if we've been given a place where to put them.
+	UsdShadeConversionImpl::ResolveTextureParameters(Prim, Material, TexturesCache, MaterialData, bShareAssetsForIdenticalPrims);
 
-	// Sort primvars
-	TSet<FString> UsedPrimvars;
-	TFunction<void(const FString& Parameter, const TOptional<UsdShadeConversionImpl::FParameterValue>&)> FetchPrimvar =
-		[&UsedPrimvars, &ParameterToPrimvar](const FString& Parameter, const TOptional<UsdShadeConversionImpl::FParameterValue>& ParameterValue)
+	// auto so we can use the MaterialInput generic parameter for FColorMaterialInput, FScalarMaterialInput and FVectorMaterialInput
+	auto ConnectMaterialInput = [&MaterialData](UMaterial& Material, auto& MaterialInput, const FString& ParameterName) -> bool
 	{
-		if (!ParameterValue.IsSet())
-		{
-			return;
-		}
-		if (const UsdShadeConversionImpl::FTextureParameterValue*
-				TextureParameterValue = ParameterValue.GetValue().TryGet<UsdShadeConversionImpl::FTextureParameterValue>())
+		const UsdToUnreal::FParameterValue* ParameterValue = MaterialData.Parameters.Find(ParameterName);
+		if (!ParameterValue)
 		{
-			UsedPrimvars.Add(TextureParameterValue->Primvar);
-			ParameterToPrimvar.Add(Parameter, TextureParameterValue->Primvar);
+			return false;
 		}
-	};
-	FetchPrimvar(TEXT("BaseColor"), BaseColorParameter);
-	FetchPrimvar(TEXT("EmissiveColor"), EmissiveParameter);
-	FetchPrimvar(TEXT("Metallic"), MetallicParameter);
-	FetchPrimvar(TEXT("Roughness"), RoughnessParameter);
-	FetchPrimvar(TEXT("Opacity"), OpacityParameter);
-	FetchPrimvar(TEXT("Normal"), NormalParameter);
-	FetchPrimvar(TEXT("Refraction"), RefractionParameter);
-	FetchPrimvar(TEXT("AmpientOcclusion"), AmbientOcclusionParameter);
-	UsedPrimvars.Remove(TEXT(""));
-	TArray<FString> SortedPrimvars = UsedPrimvars.Array();
-	SortedPrimvars.Sort();
-	TMap<FString, int32> PrimvarToUVIndex;
-	PrimvarToUVIndex.Reserve(SortedPrimvars.Num());
-	int32 UVIndex = 0;
-	for (const FString& Primvar : SortedPrimvars)
-	{
-		PrimvarToUVIndex.Add(Primvar, UVIndex++);
-	}
 
-	// auto so we can use the MaterialInput generic parameter for FColorMaterialInput, FScalarMaterialInput and FVectorMaterialInput
-	auto ConnectMaterialInput =
-		[&PrimvarToUVIndex](UMaterial& Material, auto& MaterialInput, const UsdShadeConversionImpl::FParameterValue& InParameterValue) -> bool
-	{
-		UMaterialExpression* Expression = UsdShadeConversionImpl::GetExpressionForValue(Material, InParameterValue);
+		UMaterialExpression* Expression = UsdShadeConversionImpl::GetExpressionForValue(Material, *ParameterValue);
 		if (!Expression)
 		{
 			return false;
@@ -2562,14 +2355,13 @@ bool UsdToUnreal::ConvertMaterial(
 
 		MaterialInput.Expression = Expression;
 
-		if (const UsdShadeConversionImpl::FTextureParameterValue*
-				TextureParameterValue = InParameterValue.TryGet<UsdShadeConversionImpl::FTextureParameterValue>())
+		if (const UsdToUnreal::FTextureParameterValue* TextureParameterValue = ParameterValue->TryGet<UsdToUnreal::FTextureParameterValue>())
 		{
 			MaterialInput.OutputIndex = TextureParameterValue->OutputIndex;
 
 			if (UMaterialExpressionTextureSample* TextureExpression = Cast<UMaterialExpressionTextureSample>(Expression))
 			{
-				if (int32* FoundCoordinate = PrimvarToUVIndex.Find(TextureParameterValue->Primvar))
+				if (int32* FoundCoordinate = MaterialData.PrimvarToUVIndex.Find(TextureParameterValue->Primvar))
 				{
 					TextureExpression->ConstCoordinate = *FoundCoordinate;
 				}
@@ -2580,7 +2372,7 @@ bool UsdToUnreal::ConvertMaterial(
 						Warning,
 						TEXT("Failed to find primvar '%s' when setting material parameter. Available primvars and UV indices: %s.%s"),
 						*TextureParameterValue->Primvar,
-						*UsdUtils::StringifyMap(PrimvarToUVIndex),
+						*UsdUtils::StringifyMap(MaterialData.PrimvarToUVIndex),
 						TextureParameterValue->Primvar.IsEmpty() ? TEXT(" Is your UsdUVTexture Shader missing the 'inputs:st' attribute? (It "
 																		"specifies which UV set to sample the texture with)")
 																 : TEXT("")
@@ -2594,58 +2386,17 @@ bool UsdToUnreal::ConvertMaterial(
 
 	UMaterialEditorOnlyData* EditorOnly = Material.GetEditorOnlyData();
 
-	// Base color
-	if (UsdShadeConversionImpl::FParameterValue* BaseColorParameterValue = BaseColorParameter.GetPtrOrNull())
-	{
-		ConnectMaterialInput(Material, EditorOnly->BaseColor, *BaseColorParameterValue);
-	}
-
-	// Emissive color
-	if (UsdShadeConversionImpl::FParameterValue* EmissiveParameterValue = EmissiveParameter.GetPtrOrNull())
-	{
-		ConnectMaterialInput(Material, EditorOnly->EmissiveColor, *EmissiveParameterValue);
-	}
-
-	// Metallic
-	if (UsdShadeConversionImpl::FParameterValue* MetallicParameterValue = MetallicParameter.GetPtrOrNull())
-	{
-		ConnectMaterialInput(Material, EditorOnly->Metallic, *MetallicParameterValue);
-	}
-
-	// Roughness
-	if (UsdShadeConversionImpl::FParameterValue* RoughnessParameterValue = RoughnessParameter.GetPtrOrNull())
-	{
-		ConnectMaterialInput(Material, EditorOnly->Roughness, *RoughnessParameterValue);
-	}
-
-	// Opacity
-	if (UsdShadeConversionImpl::FParameterValue* OpacityParameterValue = OpacityParameter.GetPtrOrNull())
-	{
-		const bool bConnected = ConnectMaterialInput(Material, EditorOnly->Opacity, *OpacityParameterValue);
-
-		if (bConnected)
-		{
-			Material.BlendMode = BLEND_Translucent;
-		}
-	}
-
-	// Normal
-	if (UsdShadeConversionImpl::FParameterValue* NormalParameterValue = NormalParameter.GetPtrOrNull())
+	ConnectMaterialInput(Material, EditorOnly->BaseColor, TEXT("BaseColor"));
+	ConnectMaterialInput(Material, EditorOnly->EmissiveColor, TEXT("EmissiveColor"));
+	ConnectMaterialInput(Material, EditorOnly->Metallic, TEXT("Metallic"));
+	ConnectMaterialInput(Material, EditorOnly->Roughness, TEXT("Roughness"));
+	if (ConnectMaterialInput(Material, EditorOnly->Opacity, TEXT("Opacity")))
 	{
-		ConnectMaterialInput(Material, EditorOnly->Normal, *NormalParameterValue);
-	}
-
-	// Refraction
-	if (UsdShadeConversionImpl::FParameterValue* RefractionParameterValue = RefractionParameter.GetPtrOrNull())
-	{
-		ConnectMaterialInput(Material, EditorOnly->Refraction, *RefractionParameterValue);
-	}
-
-	// Ambient occlusion
-	if (UsdShadeConversionImpl::FParameterValue* OcclusionParameterValue = AmbientOcclusionParameter.GetPtrOrNull())
-	{
-		ConnectMaterialInput(Material, EditorOnly->AmbientOcclusion, *OcclusionParameterValue);
+		Material.BlendMode = BLEND_Translucent;
 	}
+	ConnectMaterialInput(Material, EditorOnly->Normal, TEXT("Normal"));
+	ConnectMaterialInput(Material, EditorOnly->Refraction, TEXT("Refraction"));
+	ConnectMaterialInput(Material, EditorOnly->AmbientOcclusion, TEXT("AmbientOcclusion"));
 
 	// Handle world space normals
 	pxr::UsdPrim UsdPrim = UsdShadeMaterial.GetPrim();
@@ -2663,7 +2414,18 @@ bool UsdToUnreal::ConvertMaterial(
 	// values if we need to
 	if (UUsdMaterialAssetUserData* UserData = Material.GetAssetUserData<UUsdMaterialAssetUserData>())
 	{
-		UserData->PrimvarToUVIndex = PrimvarToUVIndex;
+		// We rebuild this as there's no reason to store this separately on MaterialData itself, as the complete
+		// parameter data is stored there now
+		TMap<FString, FString> ParameterToPrimvar;
+		for (const TPair<FString, UsdToUnreal::FParameterValue>& Parameter : MaterialData.Parameters)
+		{
+			if (const UsdToUnreal::FTextureParameterValue* TextureParameterValue = Parameter.Value.TryGet<UsdToUnreal::FTextureParameterValue>())
+			{
+				ParameterToPrimvar.Add(Parameter.Key, TextureParameterValue->Primvar);
+			}
+		}
+
+		UserData->PrimvarToUVIndex = MaterialData.PrimvarToUVIndex;
 		UserData->ParameterToPrimvar = ParameterToPrimvar;
 	}
 
@@ -2681,9 +2443,9 @@ bool UsdToUnreal::ConvertShadeInputsToParameters(
 	bool bShareAssetsForIdenticalPrims
 )
 {
-	FScopedUsdAllocs UsdAllocs;
+	using namespace UsdShadeConversionImpl;
 
-	bool bHasMaterialInfo = false;
+	FScopedUsdAllocs UsdAllocs;
 
 	pxr::TfToken RenderContextToken = pxr::UsdShadeTokens->universalRenderContext;
 	if (RenderContext)
@@ -2699,6 +2461,7 @@ bool UsdToUnreal::ConvertShadeInputsToParameters(
 
 	pxr::UsdShadeConnectableAPI Connectable{SurfaceShader};
 
+	bool bHasMaterialInfo = false;
 	for (pxr::UsdShadeInput& ShadeInput : SurfaceShader.GetInputs())
 	{
 		FString InputName = UsdToUnreal::ConvertToken(ShadeInput.GetBaseName());
@@ -2720,32 +2483,23 @@ bool UsdToUnreal::ConvertShadeInputsToParameters(
 		}
 
 		FString DisplayName = UsdToUnreal::ConvertString(ConnectInput.GetAttr().GetDisplayName());
-
 		if (DisplayName.IsEmpty())
 		{
 			DisplayName = InputName;
 		}
 
-		UsdShadeConversionImpl::FParameterValue ParameterValue;
-
+		UsdToUnreal::FParameterValue ParameterValue;
 		const bool bForUsdPreviewSurface = false;
-		const bool bIsNormalMap = DisplayName.Contains(TEXT("normal"));
+		const TextureGroup Group = DisplayName.Contains(TEXT("normal")) ? TEXTUREGROUP_WorldNormalMap : TEXTUREGROUP_World;
 
 		// For now it seems we don't set "primvar parameters" anyway, so don't bother building this up
 		TMap<FString, int32> Unused;
 
 		if (ShadeInput.GetTypeName() == pxr::SdfValueTypeNames->Bool)
 		{
-			if (UsdShadeConversionImpl::GetBoolParameterValue(
-					Connectable,
-					ShadeInput.GetBaseName(),
-					false,
-					ParameterValue,
-					&MaterialInstance,
-					TexturesCache
-				))
+			if (GetBoolParameterValue(Connectable, ShadeInput.GetBaseName(), false, ParameterValue))
 			{
-				UsdShadeConversionImpl::SetParameterValue(MaterialInstance, *DisplayName, ParameterValue, bForUsdPreviewSurface, Unused);
+				SetParameterValue(MaterialInstance, *DisplayName, ParameterValue, bForUsdPreviewSurface, Unused);
 				bHasMaterialInfo = true;
 			}
 		}
@@ -2753,32 +2507,15 @@ bool UsdToUnreal::ConvertShadeInputsToParameters(
 			ShadeInput.GetTypeName() == pxr::SdfValueTypeNames->Double ||
 			ShadeInput.GetTypeName() == pxr::SdfValueTypeNames->Half )
 		{
-			if (UsdShadeConversionImpl::GetFloatParameterValue(
-					Connectable,
-					ShadeInput.GetBaseName(),
-					1.f,
-					ParameterValue,
-					&MaterialInstance,
-					TexturesCache,
-					bShareAssetsForIdenticalPrims
-				))
-			{
-				UsdShadeConversionImpl::SetParameterValue(MaterialInstance, *DisplayName, ParameterValue, bForUsdPreviewSurface, Unused);
+			if (GetFloatParameterValue(Connectable, ShadeInput.GetBaseName(), 1.f, ParameterValue))
+			{
+				SetParameterValue(MaterialInstance, *DisplayName, ParameterValue, bForUsdPreviewSurface, Unused);
 				bHasMaterialInfo = true;
 			}
 		}
-		else if (UsdShadeConversionImpl::GetVec3ParameterValue(
-					 Connectable,
-					 ShadeInput.GetBaseName(),
-					 FLinearColor(0.f, 0.f, 0.f),
-					 ParameterValue,
-					 bIsNormalMap,
-					 &MaterialInstance,
-					 TexturesCache,
-					 bShareAssetsForIdenticalPrims
-				 ))
+		else if (GetVec3ParameterValue(Connectable, ShadeInput.GetBaseName(), FLinearColor::Black, ParameterValue, Group))
 		{
-			UsdShadeConversionImpl::SetParameterValue(MaterialInstance, *DisplayName, ParameterValue, bForUsdPreviewSurface, Unused);
+			SetParameterValue(MaterialInstance, *DisplayName, ParameterValue, bForUsdPreviewSurface, Unused);
 			bHasMaterialInfo = true;
 		}
 	}
@@ -3013,7 +2750,14 @@ UTexture* UsdUtils::CreateTexture(const pxr::UsdAttribute& TextureAssetPathAttr,
 	return UsdUtils::CreateTexture(ResolvedTexturePath, TextureName, Group, Flags, Outer);
 }
 
-UTexture* UsdUtils::CreateTexture(const FString& ResolvedTexturePath, FName SanitizedName, TextureGroup Group, EObjectFlags Flags, UObject* Outer, bool bForceLinear)
+UTexture* UsdUtils::CreateTexture(
+	const FString& ResolvedTexturePath,
+	FName SanitizedName,
+	TextureGroup Group,
+	EObjectFlags Flags,
+	UObject* Outer,
+	bool bForceLinear
+)
 {
 	TRACE_CPUPROFILER_EVENT_SCOPE(UsdUtils::CreateTexture);
 
@@ -3479,25 +3223,41 @@ bool UsdUtils::IsMaterialTranslucent(const pxr::UsdShadeMaterial& UsdShadeMateri
 	const UMaterialInterface* Material = nullptr;
 	UUsdAssetCache3* TexturesCache = nullptr;
 	const bool bShareAssetsForIdenticalPrims = true;
-	UsdShadeConversionImpl::FParameterValue ParameterValue;
-	bool bHasOpacityConnection = UsdShadeConversionImpl::GetFloatParameterValue(
-		Connectable,
-		UnrealIdentifiers::Opacity,
-		1.f,
-		ParameterValue,
-		Material,
-		TexturesCache,
-		bShareAssetsForIdenticalPrims
-	);
+	UsdToUnreal::FParameterValue ParameterValue;
+	bool bHasOpacityConnection = UsdShadeConversionImpl::GetFloatParameterValue(Connectable, UnrealIdentifiers::Opacity, 1.f, ParameterValue);
 
 	// Don't check if the texture is nullptr here as we won't actually parse it yet. If the variant has this type we know it's meant to be bound to a
 	// texture
-	const bool bHasBoundTexture = ParameterValue.IsType<UsdShadeConversionImpl::FTextureParameterValue>();
+	const bool bHasBoundTexture = ParameterValue.IsType<UsdToUnreal::FTextureParameterValue>();
 	const bool bIsTranslucentFloat = (ParameterValue.IsType<float>() && !FMath::IsNearlyEqual(ParameterValue.Get<float>(), 1.0f));
 
 	return bHasOpacityConnection && (bIsTranslucentFloat || bHasBoundTexture);
 }
 
+bool UsdUtils::IsMaterialTranslucent(const UsdToUnreal::FUsdPreviewSurfaceMaterialData& ConvertedMaterial)
+{
+	const UsdToUnreal::FParameterValue* OpacityParam = ConvertedMaterial.Parameters.Find(TEXT("Opacity"));
+	if (!OpacityParam)
+	{
+		return false;
+	}
+
+	if (OpacityParam->IsType<UsdToUnreal::FTextureParameterValue>())
+	{
+		return true;
+	}
+
+	if (const float* FloatValue = OpacityParam->TryGet<float>())
+	{
+		if (!FMath::IsNearlyEqual(*FloatValue, 1.0f))
+		{
+			return true;
+		}
+	}
+
+	return false;
+}
+
 FSHAHash UsdUtils::HashShadeMaterial(const pxr::UsdShadeMaterial& UsdShadeMaterial, const pxr::TfToken& RenderContext)
 {
 	FSHAHash OutHash;
diff --git a/Engine/Plugins/Runtime/USDCore/Source/USDUtilities/Public/USDGeomMeshConversion.h b/Engine/Plugins/Runtime/USDCore/Source/USDUtilities/Public/USDGeomMeshConversion.h
--- a/Engine/Plugins/Runtime/USDCore/Source/USDUtilities/Public/USDGeomMeshConversion.h
+++ b/Engine/Plugins/Runtime/USDCore/Source/USDUtilities/Public/USDGeomMeshConversion.h
@@ -9,7 +9,6 @@
 #include "UsdWrappers/ForwardDeclarations.h"
 
 #if USE_USD_SDK
-
 #include "USDIncludesStart.h"
 #include "pxr/usd/usd/timeCode.h"
 #include "pxr/usd/usdShade/tokens.h"
@@ -24,6 +23,7 @@ PXR_NAMESPACE_OPEN_SCOPE
 	class UsdShadeMaterial;
 	class UsdTyped;
 PXR_NAMESPACE_CLOSE_SCOPE
+#endif	  // USE_USD_SDK
 
 class UGeometryCache;
 class UMaterialInstanceConstant;
@@ -49,6 +49,7 @@ namespace MeshAttribute::VertexInstance
 	extern USDUTILITIES_API const FName Velocity;
 }
 
+#if USE_USD_SDK
 namespace UsdToUnreal
 {
 	/** Common options to mesh conversion functions */
@@ -113,6 +114,12 @@ namespace UsdToUnreal
 		UsdUtils::FUsdPrimMaterialAssignmentInfo& InOutMaterialAssignments,
 		const FUsdMeshConversionOptions& CommonOptions = FUsdMeshConversionOptions::DefaultOptions
 	);
+	USDUTILITIES_API bool ConvertGeomMesh(
+		const pxr::UsdPrim& MeshPrim,
+		FMeshDescription& InOutMeshDescription,
+		UsdUtils::FUsdPrimMaterialAssignmentInfo& InOutMaterialAssignments,
+		const FUsdMeshConversionOptions& CommonOptions = FUsdMeshConversionOptions::DefaultOptions
+	);
 
 	/** Bake a point instancer prim into a single FMeshDescription and MaterialAssignments, manually "instancing" all
 	 *  the prototype instances into a single combined mesh.
@@ -266,6 +273,10 @@ namespace UsdUtils
 	 */
 	struct FUsdPrimMaterialSlot
 	{
+		// This carries the name of the UsdGeomSubset if applicable, or just "Main" for the single material slot
+		// directly on the Mesh prim
+		FString SlotName;
+
 		/**
 		 * Path to the prims that contain this material assignment (e.g. '/Root/my_cube' or '/Root/my_cube/geomsubset_0'.
 		 * This is a set because this single material slot in UE may be merged from N identical prims.
@@ -279,17 +290,29 @@ namespace UsdUtils
 
 		friend bool operator==(const FUsdPrimMaterialSlot& Lhs, const FUsdPrimMaterialSlot& Rhs)
 		{
-			return Lhs.AssignmentType == Rhs.AssignmentType && Lhs.MaterialSource.Equals(Rhs.MaterialSource, ESearchCase::CaseSensitive)
-				   && Lhs.bMeshIsDoubleSided == Rhs.bMeshIsDoubleSided;
+			// Note: We intentionally don't compare the slot name itself, so that otherwise identical
+			// slots with different names can be merged together
+			return Lhs.AssignmentType == Rhs.AssignmentType &&									   //
+				   Lhs.MaterialSource.Equals(Rhs.MaterialSource, ESearchCase::CaseSensitive) &&	   //
+				   Lhs.bMeshIsDoubleSided == Rhs.bMeshIsDoubleSided;
 		}
 
 		friend uint32 GetTypeHash(const FUsdPrimMaterialSlot& Slot)
 		{
-			return HashCombine(Slot.bMeshIsDoubleSided, HashCombine(GetTypeHash(Slot.MaterialSource), static_cast<uint32>(Slot.AssignmentType)));
+			// Note: We intentionally don't hash the slot name itself, so that otherwise identical
+			// slots with different names can be merged together
+			return HashCombine(
+				Slot.bMeshIsDoubleSided,	//
+				HashCombine(				//
+					GetTypeHash(Slot.MaterialSource),
+					static_cast<uint32>(Slot.AssignmentType)
+				)
+			);
 		}
 
 		friend FArchive& operator<<(FArchive& Ar, FUsdPrimMaterialSlot& Slot)
 		{
+			Ar << Slot.SlotName;
 			Ar << Slot.MaterialSource;
 			Ar << Slot.AssignmentType;
 			Ar << Slot.bMeshIsDoubleSided;
@@ -311,6 +334,8 @@ namespace UsdUtils
 		TMap<FString, int32> PrimvarToUVIndex;
 	};
 
+	USDUTILITIES_API void RepairNormalsAndTangents(const FString& PrimPath, FMeshDescription& MeshDescription);
+
 #if USE_USD_SDK
 	/** Creates a FDisplayColorMaterial object describing the vertex color/opacity data from Gprim at time TimeCode */
 	USDUTILITIES_API TOptional<FDisplayColorMaterial> ExtractDisplayColorMaterial(
@@ -449,6 +474,5 @@ namespace UsdUtils
 	USDUTILITIES_API uint64 GetGprimVertexCount(const pxr::UsdGeomGprim& Gprim, double TimeCode);
 
 	USDUTILITIES_API void AuthorIdentityTransformGprimAttributes(const pxr::UsdPrim& Gprim, bool bDefaultValues, bool bTimeSampleValues);
-
 #endif	  // #if USE_USD_SDK
 }	 // namespace UsdUtils
diff --git a/Engine/Plugins/Runtime/USDCore/Source/USDUtilities/Public/USDShadeConversion.h b/Engine/Plugins/Runtime/USDCore/Source/USDUtilities/Public/USDShadeConversion.h
--- a/Engine/Plugins/Runtime/USDCore/Source/USDUtilities/Public/USDShadeConversion.h
+++ b/Engine/Plugins/Runtime/USDCore/Source/USDUtilities/Public/USDShadeConversion.h
@@ -7,6 +7,8 @@
 #include "CoreMinimal.h"
 #include "Engine/TextureDefines.h"
 #include "Materials/MaterialInterface.h"
+#include "Math/TransformCalculus2D.h"
+#include "Misc/TVariant.h"
 #include "UObject/Package.h"
 #include "UObject/WeakObjectPtr.h"
 
@@ -45,29 +47,84 @@ namespace UE
 
 namespace UsdToUnreal
 {
+	struct USDUTILITIES_API FTextureParameterValue
+	{
+		UTexture* Texture = nullptr;	  // Only used for the ConvertMaterial overloads that receive TexturesCache
+
+		// Parameters of the texture asset itself
+		FString TextureFilePath;
+		TextureGroup Group = TEXTUREGROUP_World;
+		TOptional<bool> bSRGB;
+		bool bIsUDIM = false;
+		TextureAddress AddressX = TA_Wrap;
+		TextureAddress AddressY = TA_Wrap;
+
+		// Parameters about the texture usage
+		FString Primvar;
+		int32 OutputIndex = 0;
+		FVector2f UVTranslation;
+		float UVRotation = 0.0f;
+		FScale2f UVScale;
+
+	public:
+		// Returns whether the texture should be parsed as sRGB or not, given the actually
+		// authored "bSRGB" member and the fallback opinion provided by the texture group
+		bool GetSRGBValue() const;
+	};
+
+	struct USDUTILITIES_API FPrimvarReaderParameterValue
+	{
+		FString PrimvarName;
+		FVector FallbackValue;
+	};
+
+	using FParameterValue = TVariant<float, FVector, FTextureParameterValue, FPrimvarReaderParameterValue, bool>;
+
+	struct USDUTILITIES_API FUsdPreviewSurfaceMaterialData
+	{
+		TMap<FString, FParameterValue> Parameters;
+
+		/**
+		 * Describes which UV set this material will target with each primvar e.g. {'firstPrimvar': 0, 'st': 1, 'st1': 2}.
+		 *
+		 * We store this here because deciding this assignment involves combining and sorting all the existing primvars
+		 * that the texture parameters want to read, which we do when first calling ConvertMaterial().
+		 *
+		 * This will later be compared with the primvar to UV index mapping we generate when parsing mesh data. If they
+		 * are compatible, we'll be able to use the material directly on that mesh. Otherwise we'll need to generate a
+		 * new instance of this material that assigns different primvars to each UV index (check
+		 * CreatePrimvarCompatibleVersionOfMaterial).
+		 */
+		TMap<FString, int32> PrimvarToUVIndex;
+	};
+
 	/**
 	 * Extracts material data from UsdShadeMaterial and places the results in Material. Note that since this is used for UMaterialInstanceDynamics at
 	 * runtime as well, it will not set base property overrides (e.g. BlendMode) or the parent material, and will just assume that the caller handles
 	 * that. Note that in order to receive the primvar to UV index mapping calculated within this function, the provided Material should have an
 	 * UUsdMaterialAssetImportData object as its AssetImportData.
 	 * @param UsdShadeMaterial - Shade material with the data to convert
-	 * @param Material - Output parameter that will be filled with the converted data. Only the versions that receive a dynamic material instance will
-	 * work at runtime
+	 * @param OutMaterial - Output parameter that will be filled with the converted data
 	 * @param TexturesCache - Cache to prevent importing a texture more than once
 	 * @param RenderContext - Which render context output to read from the UsdShadeMaterial
 	 * @param ReuseIdenticalAssets - Whether to reuse identical textures found in the TextureCache or to create dedicated textures for each material
 	 * @return Whether the conversion was successful or not.
 	 */
+	USDUTILITIES_API bool ConvertMaterial(
+		const pxr::UsdPrim& InUsdShadeMaterialPrim,
+		FUsdPreviewSurfaceMaterialData& OutMaterial,
+		const TCHAR* InRenderContext = nullptr
+	);
 	USDUTILITIES_API bool ConvertMaterial(
 		const pxr::UsdShadeMaterial& UsdShadeMaterial,
-		UMaterialInstance& Material,
+		UMaterialInstance& OutMaterial,
 		UUsdAssetCache3* TexturesCache = nullptr,
 		const TCHAR* RenderContext = nullptr,
 		bool bShareAssetsForIdenticalPrims = true
 	);
 	USDUTILITIES_API bool ConvertMaterial(
 		const pxr::UsdShadeMaterial& UsdShadeMaterial,
-		UMaterial& Material,
+		UMaterial& OutMaterial,
 		UUsdAssetCache3* TexturesCache = nullptr,
 		const TCHAR* RenderContext = nullptr,
 		bool bShareAssetsForIdenticalPrims = true
@@ -118,7 +175,7 @@ namespace UsdToUnreal
 		bool bShareAssetsForIdenticalPrims = true
 	);
 	PRAGMA_ENABLE_DEPRECATION_WARNINGS
-}
+}	 // namespace UsdToUnreal
 
 #if WITH_EDITOR
 namespace UnrealToUsd
@@ -203,6 +260,7 @@ namespace UsdUtils
 	 * This function exists because we need this information *before* we pick the right parent for a material instance and properly convert it.
 	 */
 	USDUTILITIES_API bool IsMaterialTranslucent(const pxr::UsdShadeMaterial& UsdShadeMaterial);
+	USDUTILITIES_API bool IsMaterialTranslucent(const UsdToUnreal::FUsdPreviewSurfaceMaterialData& ConvertedMaterial);
 
 	USDUTILITIES_API FSHAHash HashShadeMaterial(	//
 		const pxr::UsdShadeMaterial& UsdShadeMaterial,
