USD: Fix new call site to StaticMeshAttributes.Register() wiping MeshDescription during collapsing.

diff --git a/Engine/Plugins/Runtime/USDCore/Source/USDUtilities/Private/USDGeomMeshConversion.cpp b/Engine/Plugins/Runtime/USDCore/Source/USDUtilities/Private/USDGeomMeshConversion.cpp
--- a/Engine/Plugins/Runtime/USDCore/Source/USDUtilities/Private/USDGeomMeshConversion.cpp
+++ b/Engine/Plugins/Runtime/USDCore/Source/USDUtilities/Private/USDGeomMeshConversion.cpp
@@ -2242,7 +2242,11 @@ namespace UE::UsdGeomMeshConversion::Private
 		const int32 VertexInstanceOffset = OutMeshDescription.VertexInstances().Num();
 
 		FStaticMeshAttributes StaticMeshAttributes(OutMeshDescription);
-		StaticMeshAttributes.Register();
+
+		// Keep existing because we may be calling this in the context of collapsing, where our OutMeshDescription
+		// already has some data
+		const bool bKeepExistingAttribute = true;
+		StaticMeshAttributes.Register(bKeepExistingAttribute);
 
 		// Vertex positions
 		TVertexAttributesRef<FVector3f> MeshDescriptionVertexPositions = StaticMeshAttributes.GetVertexPositions();
@@ -2843,7 +2847,11 @@ bool UsdToUnreal::ConvertPointInstancerToMesh(
 
 	// Double-check our target mesh description has the attributes we need
 	FStaticMeshAttributes StaticMeshAttributes(OutMeshDescription);
-	StaticMeshAttributes.Register();
+
+	// Keep existing because we may be calling this in the context of collapsing, where our OutMeshDescription
+	// already has some data
+	const bool bKeepExistingAttribute = true;
+	StaticMeshAttributes.Register(bKeepExistingAttribute);
 
 	// Append mesh descriptions
 	FUsdStageInfo StageInfo{PointInstancer.GetPrim().GetStage()};
@@ -2916,7 +2924,11 @@ bool UsdToUnreal::ConvertGeomMeshHierarchy(
 	}
 
 	FStaticMeshAttributes StaticMeshAttributes(OutMeshDescription);
-	StaticMeshAttributes.Register();
+
+	// Keep existing because we may be calling this in the context of collapsing, where our OutMeshDescription
+	// already has some data
+	const bool bKeepExistingAttribute = true;
+	StaticMeshAttributes.Register(bKeepExistingAttribute);
 
 	// Pass a copy down so that we can repeatedly overwrite the AdditionalTransform and still
 	// provide the options object to ConvertGeomMesh and ConvertPointInstancerToMesh
