USD: Fix build when editor is not available.

diff --git a/Engine/Plugins/Interchange/Runtime/Source/Import/Private/Usd/InterchangeUsdTranslator.cpp b/Engine/Plugins/Interchange/Runtime/Source/Import/Private/Usd/InterchangeUsdTranslator.cpp
--- a/Engine/Plugins/Interchange/Runtime/Source/Import/Private/Usd/InterchangeUsdTranslator.cpp
+++ b/Engine/Plugins/Interchange/Runtime/Source/Import/Private/Usd/InterchangeUsdTranslator.cpp
@@ -1264,6 +1264,7 @@ namespace UE::InterchangeUsdTranslator::Private
 		TArray<FString>& OutJointNames
 	)
 	{
+#if WITH_EDITOR
 		const FString& PrimPath = PayloadKey;
 		UE::FUsdPrim Prim = Impl.UsdStage.GetPrimAtPath(UE::FSdfPath{*PrimPath});
 		if (!Prim)
@@ -1339,6 +1340,9 @@ namespace UE::InterchangeUsdTranslator::Private
 		OutJointNames = *JointNames;
 
 		return true;
+#else
+		return false;
+#endif	  // WITH_EDITOR
 	}
 
 	bool GetMorphTargetPayloadData(
diff --git a/Engine/Plugins/Runtime/USDCore/Source/USDUtilities/Public/USDSkeletalDataConversion.h b/Engine/Plugins/Runtime/USDCore/Source/USDUtilities/Public/USDSkeletalDataConversion.h
--- a/Engine/Plugins/Runtime/USDCore/Source/USDUtilities/Public/USDSkeletalDataConversion.h
+++ b/Engine/Plugins/Runtime/USDCore/Source/USDUtilities/Public/USDSkeletalDataConversion.h
@@ -192,7 +192,6 @@ namespace UsdUtils
 #endif	  // USE_USD_SDK
 }	 // namespace UsdUtils
 
-#if USE_USD_SDK && WITH_EDITOR
 namespace UsdToUnreal
 {
 	struct FUsdSkeletonData
@@ -210,6 +209,7 @@ namespace UsdToUnreal
 		TArray<FBone> Bones;
 	};
 
+#if USE_USD_SDK && WITH_EDITOR
 	/**
 	 * Extracts skeleton data from UsdSkeletonQuery and places the results in ConvertedData.
 	 * @param UsdSkeletonQuery - SkeletonQuery with the data to convert
@@ -363,8 +363,11 @@ namespace UsdToUnreal
 		UsdUtils::FBlendShapeMap& InBlendShapesByPath,
 		USkeletalMesh* InOutSkeletalMesh
 	);
+
+#endif	  // #if USE_USD_SDK && WITH_EDITOR
 }	 // namespace UsdToUnreal
 
+#if USE_USD_SDK && WITH_EDITOR
 namespace UnrealToUsd
 {
 	/**
@@ -442,5 +445,4 @@ namespace UnrealToUsd
 		const UsdUtils::FBlendShapeMap* InBlendShapeMap = nullptr
 	);
 }
-
 #endif	  // #if USE_USD_SDK && WITH_EDITOR
