USD: Fix a couple of trivial mistakes with cl 34703718 causing skeletal mesh imports to fail.

diff --git a/Engine/Plugins/Runtime/USDCore/Source/USDUtilities/Private/USDSkeletalDataConversion.cpp b/Engine/Plugins/Runtime/USDCore/Source/USDUtilities/Private/USDSkeletalDataConversion.cpp
--- a/Engine/Plugins/Runtime/USDCore/Source/USDUtilities/Private/USDSkeletalDataConversion.cpp
+++ b/Engine/Plugins/Runtime/USDCore/Source/USDUtilities/Private/USDSkeletalDataConversion.cpp
@@ -1271,13 +1271,14 @@ bool UsdToUnreal::ConvertSkeleton(
 				FText::FromString(SkeletonPrimPath)
 			)
 		);
-	}
-	if (bEnsureAtLeastOneBone)
-	{
-		FUsdSkeletonData::FBone& Bone = OutConvertedData.Bones.Emplace_GetRef();
-		Bone.Name = TEXT("Root");
-		Bone.ParentIndex = INDEX_NONE;
-		return true;
+
+		if (bEnsureAtLeastOneBone)
+		{
+			FUsdSkeletonData::FBone& Bone = OutConvertedData.Bones.Emplace_GetRef();
+			Bone.Name = TEXT("Root");
+			Bone.ParentIndex = INDEX_NONE;
+			return true;
+		}
 	}
 
 	// Get the rest transforms that will end up on the skeleton data.
@@ -3149,7 +3150,7 @@ UE::FUsdSkelSkinningQuery UsdUtils::CreateSkinningQuery(const pxr::UsdPrim& Skin
 {
 	pxr::UsdSkelBindingAPI SkelBindingAPI{SkinnedMeshPrim};
 	const pxr::UsdSkelAnimQuery& AnimQuery = SkeletonQuery.GetAnimQuery();
-	if (!SkelBindingAPI || !AnimQuery)
+	if (!SkelBindingAPI)
 	{
 		return {};
 	}
