USD: Fix ensure when trying to handle any type of unsupported track as draw mode tracks.

diff --git a/Engine/Plugins/Runtime/USDCore/Source/USDUtilities/Private/USDPrimConversion.cpp b/Engine/Plugins/Runtime/USDCore/Source/USDUtilities/Private/USDPrimConversion.cpp
--- a/Engine/Plugins/Runtime/USDCore/Source/USDUtilities/Private/USDPrimConversion.cpp
+++ b/Engine/Plugins/Runtime/USDCore/Source/USDUtilities/Private/USDPrimConversion.cpp
@@ -5080,8 +5080,7 @@ UnrealToUsd::FPropertyTrackWriter UnrealToUsd::CreatePropertyTrackWriter(
 				}
 			}
 		}
-		// Bounds component properties
-		else
+		else if (const UUsdDrawModeComponent* DrawModeComponent = Cast<const UUsdDrawModeComponent>(&Component))
 		{
 			Attr = {};
 			if (pxr::UsdGeomModelAPI GeomModelAPI{UsdPrim})
