USD: Fix missing invalid prim check.

diff --git a/Engine/Plugins/Runtime/USDCore/Source/USDUtilities/Private/USDLightConversion.cpp b/Engine/Plugins/Runtime/USDCore/Source/USDUtilities/Private/USDLightConversion.cpp
--- a/Engine/Plugins/Runtime/USDCore/Source/USDUtilities/Private/USDLightConversion.cpp
+++ b/Engine/Plugins/Runtime/USDCore/Source/USDUtilities/Private/USDLightConversion.cpp
@@ -783,6 +783,11 @@ bool UnrealToUsd::ConvertSpotLightComponent(const USpotLightComponent& LightComp
 {
 	FScopedUsdAllocs Allocs;
 
+	if (!Prim)
+	{
+		return false;
+	}
+
 	pxr::UsdLuxShapingAPI ShapingAPI = pxr::UsdLuxShapingAPI::Apply(Prim);
 	if (!ShapingAPI)
 	{
