USD: Use infinite size for generated LevelSequence sections, which behaves better especially for bool/visibility tracks;

diff --git a/Engine/Plugins/Runtime/USDCore/Source/USDUtilities/Private/USDPrimConversion.cpp b/Engine/Plugins/Runtime/USDCore/Source/USDUtilities/Private/USDPrimConversion.cpp
--- a/Engine/Plugins/Runtime/USDCore/Source/USDUtilities/Private/USDPrimConversion.cpp
+++ b/Engine/Plugins/Runtime/USDCore/Source/USDUtilities/Private/USDPrimConversion.cpp
@@ -598,7 +598,7 @@ bool UsdToUnreal::ConvertBoolTimeSamples(
 		Data.AddKey(FrameNumbers[KeyIndex], SectionValues[KeyIndex]);
 	}
 
-	Section->SetRange(Section->GetAutoSizeRange().Get(TRange<FFrameNumber>::Empty()));
+	Section->SetRange(TRange<FFrameNumber>::All());
 
 	return true;
 }
@@ -673,7 +673,7 @@ bool UsdToUnreal::ConvertBoolTimeSamples(
 		Data.AddKey(FrameNumbers[KeyIndex], SectionValues[KeyIndex]);
 	}
 
-	Section->SetRange(Section->GetAutoSizeRange().Get(TRange<FFrameNumber>::Empty()));
+	Section->SetRange(TRange<FFrameNumber>::All());
 
 	return true;
 }
@@ -788,7 +788,8 @@ bool UsdToUnreal::ConvertFloatTimeSamples(
 		}
 	}
 
-	Section->SetRange(Section->GetAutoSizeRange().Get(TRange<FFrameNumber>::Empty()));
+	Section->SetRange(TRange<FFrameNumber>::All());
+
 	return true;
 }
 
@@ -879,7 +880,7 @@ bool UsdToUnreal::ConvertColorTimeSamples(
 	Channels[2]->Set(FrameNumbers, BlueValues);
 	Channels[3]->Set(FrameNumbers, AlphaValues);
 
-	Section->SetRange(Section->GetAutoSizeRange().Get(TRange<FFrameNumber>::Empty()));
+	Section->SetRange(TRange<FFrameNumber>::All());
 
 	return true;
 }
@@ -1149,8 +1150,8 @@ bool UsdToUnreal::ConvertBoundsTimeSamples(
 	MaxChannels[1]->Set(FrameNumbers, MaxYValues);
 	MaxChannels[2]->Set(FrameNumbers, MaxZValues);
 
-	MinSection->SetRange(MinSection->GetAutoSizeRange().Get(TRange<FFrameNumber>::Empty()));
-	MaxSection->SetRange(MaxSection->GetAutoSizeRange().Get(TRange<FFrameNumber>::Empty()));
+	MinSection->SetRange(TRange<FFrameNumber>::All());
+	MaxSection->SetRange(TRange<FFrameNumber>::All());
 
 	return true;
 }
