USD: Fix exported LevelSequence transforms when using top-level actors and attachment tracks.

diff --git a/Engine/Plugins/Runtime/USDCore/Source/USDUtilities/Private/USDPrimConversion.cpp b/Engine/Plugins/Runtime/USDCore/Source/USDUtilities/Private/USDPrimConversion.cpp
--- a/Engine/Plugins/Runtime/USDCore/Source/USDUtilities/Private/USDPrimConversion.cpp
+++ b/Engine/Plugins/Runtime/USDCore/Source/USDUtilities/Private/USDPrimConversion.cpp
@@ -4069,7 +4069,9 @@ bool UnrealToUsd::CreateComponentPropertyBaker(
 				}
 				else
 				{
-					RelativeTransform = Component.GetRelativeTransform();
+					// Use the world transform here, because while we may not have an *original* attach parent, this sequence
+					// could have attach tracks, meaning that we may gain a different attach parent at some point
+					RelativeTransform = Component.GetComponentTransform();
 				}
 
 				RelativeTransform = CameraCompensation * RelativeTransform * InverseParentCameraCompensation;
