USD: Refresh component references when baking blueprints that run their construction scripts continuously in the Sequencer.

diff --git a/Engine/Plugins/Runtime/USDCore/Source/USDUtilities/Private/USDPrimConversion.cpp b/Engine/Plugins/Runtime/USDCore/Source/USDUtilities/Private/USDPrimConversion.cpp
--- a/Engine/Plugins/Runtime/USDCore/Source/USDUtilities/Private/USDPrimConversion.cpp
+++ b/Engine/Plugins/Runtime/USDCore/Source/USDUtilities/Private/USDPrimConversion.cpp
@@ -3962,6 +3962,62 @@ bool UnrealToUsd::ConvertInstancedFoliageActor(const AInstancedFoliageActor& Act
 #endif	  // WITH_EDITOR
 }
 
+namespace UE::USDPrimConversion::Private
+{
+	// If we're going to bake a LevelSequence with bindings to blueprints that have been set with bRunConstructionScriptInSequencer,
+	// on every single baked frame the construction script will be rerun, likely recreating all of the blueprint's components.
+	// This means we can't just capture component references into our baker lambdas, and must instead capture something that can
+	// find its ideal component again in that case.
+	//
+	// We're going to assume component names are consistent across every run of the construction scripts, and just use these getters
+	// to reconnect with the ideal components. The actor itself survives the construction script, as it is the blueprint actor itself
+	template<typename T>
+	TFunction<T*(void)> CreateComponentGetter(T* OldComponent)
+	{
+		if (!OldComponent)
+		{
+			return []()
+			{
+				return nullptr;
+			};
+		}
+
+		const AActor* Actor = OldComponent->GetOwner();
+		if (!Actor)
+		{
+			// Can't do much without an actor (shouldn't really ever happen though)
+			return [OldComponent]()
+			{
+				return IsValid(OldComponent) ? OldComponent : nullptr;
+			};
+		}
+
+		const FName ComponentName = OldComponent->GetFName();
+		return [OldComponent, Actor, ComponentName]() -> T*
+		{
+			if (IsValid(OldComponent))
+			{
+				return OldComponent;
+			}
+
+			// TODO: Maybe keep track of the last index where we found our component,
+			// and start searching there the next call?
+			for (UActorComponent* ChildComponent : Actor->GetComponents())
+			{
+				if (T* CastComponent = Cast<T>(ChildComponent))
+				{
+					if (ChildComponent->GetFName() == ComponentName)
+					{
+						return CastComponent;
+					}
+				}
+			}
+
+			return nullptr;
+		};
+	}
+}	 // namespace UE::USDPrimConversion::Private
+
 bool UnrealToUsd::CreateComponentPropertyBaker(
 	UE::FUsdPrim& Prim,
 	const USceneComponent& Component,
@@ -3969,6 +4025,8 @@ bool UnrealToUsd::CreateComponentPropertyBaker(
 	FComponentBaker& OutBaker
 )
 {
+	using namespace UE::USDPrimConversion::Private;
+
 	EBakingType BakerType = EBakingType::None;
 	TFunction<void(double)> BakerFunction;
 
@@ -4047,12 +4105,27 @@ bool UnrealToUsd::CreateComponentPropertyBaker(
 				}
 			}
 
+			TFunction<const USceneComponent*(void)> ComponentGetter = CreateComponentGetter(&Component);
+			TFunction<USceneComponent*(void)> OriginalAttachParentGetter = CreateComponentGetter(OriginalAttachParent);
+
 			BakerType = EBakingType::Transform;
-			BakerFunction =
-				[&Component, CameraCompensation, InverseParentCameraCompensation, StageInfo, Attr, OriginalAttachParent](double UsdTimeCode)
+			BakerFunction = [CameraCompensation,	//
+							 InverseParentCameraCompensation,
+							 StageInfo,
+							 Attr,
+							 ComponentGetter,
+							 OriginalAttachParentGetter](double UsdTimeCode)
 			{
 				FScopedUsdAllocs Allocs;
 
+				const USceneComponent* Component = ComponentGetter();
+				if (!Component)
+				{
+					return;
+				}
+
+				USceneComponent* OriginalAttachParent = OriginalAttachParentGetter();
+
 				// If we're attached to a socket our RelativeTransform will be relative to the socket, instead of the parent
 				// component space. If we were to use GetRelativeTransform directly, we're in charge of managing the socket
 				// transform too (and any other N obscure features we don't know about/don't exist yet). If we fetch directly
@@ -4065,13 +4138,13 @@ bool UnrealToUsd::CreateComponentPropertyBaker(
 				{
 					OriginalAttachParent->ConditionalUpdateComponentToWorld();
 					OriginalAttachParent->UpdateChildTransforms();
-					RelativeTransform = Component.GetComponentTransform().GetRelativeTransform(OriginalAttachParent->GetComponentTransform());
+					RelativeTransform = Component->GetComponentTransform().GetRelativeTransform(OriginalAttachParent->GetComponentTransform());
 				}
 				else
 				{
 					// Use the world transform here, because while we may not have an *original* attach parent, this sequence
 					// could have attach tracks, meaning that we may gain a different attach parent at some point
-					RelativeTransform = Component.GetComponentTransform();
+					RelativeTransform = Component->GetComponentTransform();
 				}
 
 				RelativeTransform = CameraCompensation * RelativeTransform * InverseParentCameraCompensation;
@@ -4095,10 +4168,18 @@ bool UnrealToUsd::CreateComponentPropertyBaker(
 			pxr::UsdAttribute Attr = Imageable.CreateVisibilityAttr();
 			Attr.Clear();
 
+			TFunction<const USceneComponent*(void)> ComponentGetter = CreateComponentGetter(&Component);
+
 			BakerType = EBakingType::Visibility;
-			BakerFunction = [&Component, Imageable](double UsdTimeCode)
+			BakerFunction = [Imageable, ComponentGetter](double UsdTimeCode)
 			{
-				if (Component.bHiddenInGame || Component.GetOwner()->IsHidden())
+				const USceneComponent* Component = ComponentGetter();
+				if (!Component)
+				{
+					return;
+				}
+
+				if (Component->bHiddenInGame || Component->GetOwner()->IsHidden())
 				{
 					Imageable.MakeInvisible(UsdTimeCode);
 
@@ -4143,9 +4224,17 @@ bool UnrealToUsd::CreateComponentPropertyBaker(
 
 		if (RelevantProperties.Contains(PropertyPath))
 		{
+			TFunction<const UCineCameraComponent*(void)> ComponentGetter = CreateComponentGetter(CameraComponent);
+
 			BakerType = EBakingType::Camera;
-			BakerFunction = [UsdStage, CameraComponent, UsdPrim](double UsdTimeCode) mutable
+			BakerFunction = [UsdStage, UsdPrim, ComponentGetter](double UsdTimeCode) mutable
 			{
+				const UCineCameraComponent* CameraComponent = ComponentGetter();
+				if (!CameraComponent)
+				{
+					return;
+				}
+
 				UnrealToUsd::ConvertCameraComponent(*CameraComponent, UsdPrim, UsdTimeCode);
 			};
 		}
@@ -4159,9 +4248,17 @@ bool UnrealToUsd::CreateComponentPropertyBaker(
 
 			if (RelevantProperties.Contains(PropertyPath))
 			{
+				TFunction<const URectLightComponent*(void)> ComponentGetter = CreateComponentGetter(RectLightComponent);
+
 				BakerType = EBakingType::Light;
-				BakerFunction = [UsdStage, RectLightComponent, UsdPrim](double UsdTimeCode) mutable
+				BakerFunction = [UsdStage, UsdPrim, ComponentGetter](double UsdTimeCode) mutable
 				{
+					const URectLightComponent* RectLightComponent = ComponentGetter();
+					if (!RectLightComponent)
+					{
+						return;
+					}
+
 					UnrealToUsd::ConvertLightComponent(*RectLightComponent, UsdPrim, UsdTimeCode);
 					UnrealToUsd::ConvertRectLightComponent(*RectLightComponent, UsdPrim, UsdTimeCode);
 				};
@@ -4174,9 +4271,17 @@ bool UnrealToUsd::CreateComponentPropertyBaker(
 
 			if (RelevantProperties.Contains(PropertyPath))
 			{
+				TFunction<const USpotLightComponent*(void)> ComponentGetter = CreateComponentGetter(SpotLightComponent);
+
 				BakerType = EBakingType::Light;
-				BakerFunction = [UsdStage, SpotLightComponent, UsdPrim](double UsdTimeCode) mutable
+				BakerFunction = [UsdStage, UsdPrim, ComponentGetter](double UsdTimeCode) mutable
 				{
+					const USpotLightComponent* SpotLightComponent = ComponentGetter();
+					if (!SpotLightComponent)
+					{
+						return;
+					}
+
 					UnrealToUsd::ConvertLightComponent(*SpotLightComponent, UsdPrim, UsdTimeCode);
 					UnrealToUsd::ConvertPointLightComponent(*SpotLightComponent, UsdPrim, UsdTimeCode);
 					UnrealToUsd::ConvertSpotLightComponent(*SpotLightComponent, UsdPrim, UsdTimeCode);
@@ -4190,9 +4295,17 @@ bool UnrealToUsd::CreateComponentPropertyBaker(
 
 			if (RelevantProperties.Contains(PropertyPath))
 			{
+				TFunction<const UPointLightComponent*(void)> ComponentGetter = CreateComponentGetter(PointLightComponent);
+
 				BakerType = EBakingType::Light;
-				BakerFunction = [UsdStage, PointLightComponent, UsdPrim](double UsdTimeCode) mutable
+				BakerFunction = [UsdStage, UsdPrim, ComponentGetter](double UsdTimeCode) mutable
 				{
+					const UPointLightComponent* PointLightComponent = ComponentGetter();
+					if (!PointLightComponent)
+					{
+						return;
+					}
+
 					UnrealToUsd::ConvertLightComponent(*PointLightComponent, UsdPrim, UsdTimeCode);
 					UnrealToUsd::ConvertPointLightComponent(*PointLightComponent, UsdPrim, UsdTimeCode);
 				};
@@ -4205,9 +4318,17 @@ bool UnrealToUsd::CreateComponentPropertyBaker(
 
 			if (RelevantProperties.Contains(PropertyPath))
 			{
+				TFunction<const UDirectionalLightComponent*(void)> ComponentGetter = CreateComponentGetter(DirectionalLightComponent);
+
 				BakerType = EBakingType::Light;
-				BakerFunction = [UsdStage, DirectionalLightComponent, UsdPrim](double UsdTimeCode) mutable
+				BakerFunction = [UsdStage, UsdPrim, ComponentGetter](double UsdTimeCode) mutable
 				{
+					const UDirectionalLightComponent* DirectionalLightComponent = ComponentGetter();
+					if (!DirectionalLightComponent)
+					{
+						return;
+					}
+
 					UnrealToUsd::ConvertLightComponent(*DirectionalLightComponent, UsdPrim, UsdTimeCode);
 					UnrealToUsd::ConvertDirectionalLightComponent(*DirectionalLightComponent, UsdPrim, UsdTimeCode);
 				};
@@ -4222,9 +4343,17 @@ bool UnrealToUsd::CreateComponentPropertyBaker(
 
 		if (RelevantProperties.Contains(PropertyPath))
 		{
+			TFunction<const UUsdDrawModeComponent*(void)> ComponentGetter = CreateComponentGetter(DrawModeComponent);
+
 			BakerType = EBakingType::Bounds;
-			BakerFunction = [UsdStage, DrawModeComponent, UsdPrim](double UsdTimeCode) mutable
+			BakerFunction = [UsdStage, UsdPrim, ComponentGetter](double UsdTimeCode) mutable
 			{
+				const UUsdDrawModeComponent* DrawModeComponent = ComponentGetter();
+				if (!DrawModeComponent)
+				{
+					return;
+				}
+
 				const bool bWriteExtents = true;
 				UnrealToUsd::ConvertDrawModeComponent(*DrawModeComponent, UsdPrim, bWriteExtents, UsdTimeCode);
 			};
