Optimize FUsdCardsSceneProxy to reduce the amount of force-inlined calls to FStaticMeshVertexBuffer::SetVertexTangents, which itself force-inlines a tree of calls to other force-inlined functions.
This change results in wall duration for the compile of Module.USDClasses.cpp (which includes USDDrawModeComponent.cpp) being reduced from 96.7s to 10.7s (per data gathered via vcperf).

diff --git a/Engine/Plugins/Runtime/USDCore/Source/USDClasses/Private/USDDrawModeComponent.cpp b/Engine/Plugins/Runtime/USDCore/Source/USDClasses/Private/USDDrawModeComponent.cpp
--- a/Engine/Plugins/Runtime/USDCore/Source/USDClasses/Private/USDDrawModeComponent.cpp
+++ b/Engine/Plugins/Runtime/USDCore/Source/USDClasses/Private/USDDrawModeComponent.cpp
@@ -392,82 +392,57 @@ namespace UE::UsdDrawModeComponentImpl::Private
 
 				const int NumTriangles = 2;
 				const int NumVertices = 4;
-				const int NumTexCoords = 1;
-				Buffer.PositionVertexBuffer.Init(NumVertices);
-				Buffer.StaticMeshVertexBuffer.Init(NumVertices, NumTexCoords);
+
+				FStaticMeshBuildVertex VertexBuildData[NumVertices];
 
 				switch (Face)
 				{
 					case EUsdModelCardFace::XPos:
 					{
-						Buffer.PositionVertexBuffer.VertexPosition(0) = {OffsetAlongAxis, -0.5, -0.5};
-						Buffer.PositionVertexBuffer.VertexPosition(1) = {OffsetAlongAxis, -0.5, +0.5};
-						Buffer.PositionVertexBuffer.VertexPosition(2) = {OffsetAlongAxis, +0.5, +0.5};
-						Buffer.PositionVertexBuffer.VertexPosition(3) = {OffsetAlongAxis, +0.5, -0.5};
-						Buffer.StaticMeshVertexBuffer.SetVertexTangents(0, {0, 1, 0}, {0, 0, 1}, {1, 0, 0});
-						Buffer.StaticMeshVertexBuffer.SetVertexTangents(1, {0, 1, 0}, {0, 0, 1}, {1, 0, 0});
-						Buffer.StaticMeshVertexBuffer.SetVertexTangents(2, {0, 1, 0}, {0, 0, 1}, {1, 0, 0});
-						Buffer.StaticMeshVertexBuffer.SetVertexTangents(3, {0, 1, 0}, {0, 0, 1}, {1, 0, 0});
+						VertexBuildData[0] = { .Position = {OffsetAlongAxis, -0.5, -0.5}, .TangentX = {0, 1, 0}, .TangentY = {0, 0, 1}, .TangentZ = {1, 0, 0} };
+						VertexBuildData[1] = { .Position = {OffsetAlongAxis, -0.5, +0.5}, .TangentX = {0, 1, 0}, .TangentY = {0, 0, 1}, .TangentZ = {1, 0, 0} };
+						VertexBuildData[2] = { .Position = {OffsetAlongAxis, +0.5, +0.5}, .TangentX = {0, 1, 0}, .TangentY = {0, 0, 1}, .TangentZ = {1, 0, 0} };
+						VertexBuildData[3] = { .Position = {OffsetAlongAxis, +0.5, -0.5}, .TangentX = {0, 1, 0}, .TangentY = {0, 0, 1}, .TangentZ = {1, 0, 0} };
 						break;
 					}
 					case EUsdModelCardFace::YPos:
 					{
-						Buffer.PositionVertexBuffer.VertexPosition(0) = {+0.5, OffsetAlongAxis, -0.5};
-						Buffer.PositionVertexBuffer.VertexPosition(1) = {+0.5, OffsetAlongAxis, +0.5};
-						Buffer.PositionVertexBuffer.VertexPosition(2) = {-0.5, OffsetAlongAxis, +0.5};
-						Buffer.PositionVertexBuffer.VertexPosition(3) = {-0.5, OffsetAlongAxis, -0.5};
-						Buffer.StaticMeshVertexBuffer.SetVertexTangents(0, {0, 0, 1}, {1, 0, 0}, {0, 1, 0});
-						Buffer.StaticMeshVertexBuffer.SetVertexTangents(1, {0, 0, 1}, {1, 0, 0}, {0, 1, 0});
-						Buffer.StaticMeshVertexBuffer.SetVertexTangents(2, {0, 0, 1}, {1, 0, 0}, {0, 1, 0});
-						Buffer.StaticMeshVertexBuffer.SetVertexTangents(3, {0, 0, 1}, {1, 0, 0}, {0, 1, 0});
+						VertexBuildData[0] = { .Position = {+0.5, OffsetAlongAxis, -0.5}, .TangentX = {0, 0, 1}, .TangentY = {1, 0, 0}, .TangentZ = {0, 1, 0} };
+						VertexBuildData[1] = { .Position = {+0.5, OffsetAlongAxis, +0.5}, .TangentX = {0, 0, 1}, .TangentY = {1, 0, 0}, .TangentZ = {0, 1, 0} };
+						VertexBuildData[2] = { .Position = {-0.5, OffsetAlongAxis, +0.5}, .TangentX = {0, 0, 1}, .TangentY = {1, 0, 0}, .TangentZ = {0, 1, 0} };
+						VertexBuildData[3] = { .Position = {-0.5, OffsetAlongAxis, -0.5}, .TangentX = {0, 0, 1}, .TangentY = {1, 0, 0}, .TangentZ = {0, 1, 0} };
 						break;
 					}
 					case EUsdModelCardFace::ZPos:
 					{
-						Buffer.PositionVertexBuffer.VertexPosition(0) = {+0.5, +0.5, OffsetAlongAxis};
-						Buffer.PositionVertexBuffer.VertexPosition(1) = {+0.5, -0.5, OffsetAlongAxis};
-						Buffer.PositionVertexBuffer.VertexPosition(2) = {-0.5, -0.5, OffsetAlongAxis};
-						Buffer.PositionVertexBuffer.VertexPosition(3) = {-0.5, +0.5, OffsetAlongAxis};
-						Buffer.StaticMeshVertexBuffer.SetVertexTangents(0, {1, 0, 0}, {0, 1, 0}, {0, 0, 1});
-						Buffer.StaticMeshVertexBuffer.SetVertexTangents(1, {1, 0, 0}, {0, 1, 0}, {0, 0, 1});
-						Buffer.StaticMeshVertexBuffer.SetVertexTangents(2, {1, 0, 0}, {0, 1, 0}, {0, 0, 1});
-						Buffer.StaticMeshVertexBuffer.SetVertexTangents(3, {1, 0, 0}, {0, 1, 0}, {0, 0, 1});
+						VertexBuildData[0] = { .Position = {+0.5, +0.5, OffsetAlongAxis}, .TangentX = {1, 0, 0}, .TangentY = {0, 1, 0}, .TangentZ = {0, 0, 1} };
+						VertexBuildData[1] = { .Position = {+0.5, -0.5, OffsetAlongAxis}, .TangentX = {1, 0, 0}, .TangentY = {0, 1, 0}, .TangentZ = {0, 0, 1} };
+						VertexBuildData[2] = { .Position = {-0.5, -0.5, OffsetAlongAxis}, .TangentX = {1, 0, 0}, .TangentY = {0, 1, 0}, .TangentZ = {0, 0, 1} };
+						VertexBuildData[3] = { .Position = {-0.5, +0.5, OffsetAlongAxis}, .TangentX = {1, 0, 0}, .TangentY = {0, 1, 0}, .TangentZ = {0, 0, 1} };
 						break;
 					}
 					case EUsdModelCardFace::XNeg:
 					{
-						Buffer.PositionVertexBuffer.VertexPosition(0) = {-OffsetAlongAxis, +0.5, -0.5};
-						Buffer.PositionVertexBuffer.VertexPosition(1) = {-OffsetAlongAxis, +0.5, +0.5};
-						Buffer.PositionVertexBuffer.VertexPosition(2) = {-OffsetAlongAxis, -0.5, +0.5};
-						Buffer.PositionVertexBuffer.VertexPosition(3) = {-OffsetAlongAxis, -0.5, -0.5};
-						Buffer.StaticMeshVertexBuffer.SetVertexTangents(0, {0, 0, 1}, {0, 1, 1}, {-1, 0, 0});
-						Buffer.StaticMeshVertexBuffer.SetVertexTangents(1, {0, 0, 1}, {0, 1, 1}, {-1, 0, 0});
-						Buffer.StaticMeshVertexBuffer.SetVertexTangents(2, {0, 0, 1}, {0, 1, 1}, {-1, 0, 0});
-						Buffer.StaticMeshVertexBuffer.SetVertexTangents(3, {0, 0, 1}, {0, 1, 1}, {-1, 0, 0});
+						VertexBuildData[0] = { .Position = {-OffsetAlongAxis, +0.5, -0.5}, .TangentX = {0, 0, 1}, .TangentY = {0, 1, 1}, .TangentZ = {-1, 0, 0} };
+						VertexBuildData[1] = { .Position = {-OffsetAlongAxis, +0.5, +0.5}, .TangentX = {0, 0, 1}, .TangentY = {0, 1, 1}, .TangentZ = {-1, 0, 0} };
+						VertexBuildData[2] = { .Position = {-OffsetAlongAxis, -0.5, +0.5}, .TangentX = {0, 0, 1}, .TangentY = {0, 1, 1}, .TangentZ = {-1, 0, 0} };
+						VertexBuildData[3] = { .Position = {-OffsetAlongAxis, -0.5, -0.5}, .TangentX = {0, 0, 1}, .TangentY = {0, 1, 1}, .TangentZ = {-1, 0, 0} };
 						break;
 					}
 					case EUsdModelCardFace::YNeg:
 					{
-						Buffer.PositionVertexBuffer.VertexPosition(0) = {-0.5, -OffsetAlongAxis, -0.5};
-						Buffer.PositionVertexBuffer.VertexPosition(1) = {-0.5, -OffsetAlongAxis, +0.5};
-						Buffer.PositionVertexBuffer.VertexPosition(2) = {+0.5, -OffsetAlongAxis, +0.5};
-						Buffer.PositionVertexBuffer.VertexPosition(3) = {+0.5, -OffsetAlongAxis, -0.5};
-						Buffer.StaticMeshVertexBuffer.SetVertexTangents(0, {0, 0, 1}, {-1, 0, 0}, {0, -1, 0});
-						Buffer.StaticMeshVertexBuffer.SetVertexTangents(1, {0, 0, 1}, {-1, 0, 0}, {0, -1, 0});
-						Buffer.StaticMeshVertexBuffer.SetVertexTangents(2, {0, 0, 1}, {-1, 0, 0}, {0, -1, 0});
-						Buffer.StaticMeshVertexBuffer.SetVertexTangents(3, {0, 0, 1}, {-1, 0, 0}, {0, -1, 0});
+						VertexBuildData[0] = { .Position = {-0.5, -OffsetAlongAxis, -0.5}, .TangentX = {0, 0, 1}, .TangentY = {-1, 0, 0}, .TangentZ = {0, -1, 0} };
+						VertexBuildData[1] = { .Position = {-0.5, -OffsetAlongAxis, +0.5}, .TangentX = {0, 0, 1}, .TangentY = {-1, 0, 0}, .TangentZ = {0, -1, 0} };
+						VertexBuildData[2] = { .Position = {+0.5, -OffsetAlongAxis, +0.5}, .TangentX = {0, 0, 1}, .TangentY = {-1, 0, 0}, .TangentZ = {0, -1, 0} };
+						VertexBuildData[3] = { .Position = {+0.5, -OffsetAlongAxis, -0.5}, .TangentX = {0, 0, 1}, .TangentY = {-1, 0, 0}, .TangentZ = {0, -1, 0} };
 						break;
 					}
 					case EUsdModelCardFace::ZNeg:
 					{
-						Buffer.PositionVertexBuffer.VertexPosition(0) = {+0.5, -0.5, -OffsetAlongAxis};
-						Buffer.PositionVertexBuffer.VertexPosition(1) = {+0.5, +0.5, -OffsetAlongAxis};
-						Buffer.PositionVertexBuffer.VertexPosition(2) = {-0.5, +0.5, -OffsetAlongAxis};
-						Buffer.PositionVertexBuffer.VertexPosition(3) = {-0.5, -0.5, -OffsetAlongAxis};
-						Buffer.StaticMeshVertexBuffer.SetVertexTangents(0, {1, 0, 0}, {0, -1, 0}, {0, 0, -1});
-						Buffer.StaticMeshVertexBuffer.SetVertexTangents(1, {1, 0, 0}, {0, -1, 0}, {0, 0, -1});
-						Buffer.StaticMeshVertexBuffer.SetVertexTangents(2, {1, 0, 0}, {0, -1, 0}, {0, 0, -1});
-						Buffer.StaticMeshVertexBuffer.SetVertexTangents(3, {1, 0, 0}, {0, -1, 0}, {0, 0, -1});
+						VertexBuildData[0] = { .Position = {+0.5, -0.5, -OffsetAlongAxis}, .TangentX = {1, 0, 0}, .TangentY = {0, -1, 0}, .TangentZ = {0, 0, -1} };
+						VertexBuildData[1] = { .Position = {+0.5, +0.5, -OffsetAlongAxis}, .TangentX = {1, 0, 0}, .TangentY = {0, -1, 0}, .TangentZ = {0, 0, -1} };
+						VertexBuildData[2] = { .Position = {-0.5, +0.5, -OffsetAlongAxis}, .TangentX = {1, 0, 0}, .TangentY = {0, -1, 0}, .TangentZ = {0, 0, -1} };
+						VertexBuildData[3] = { .Position = {-0.5, -0.5, -OffsetAlongAxis}, .TangentX = {1, 0, 0}, .TangentY = {0, -1, 0}, .TangentZ = {0, 0, -1} };
 						break;
 					}
 					default:
@@ -483,25 +458,41 @@ namespace UE::UsdDrawModeComponentImpl::Private
 				const bool bFlipHorizontally = !bFaceIsAuthored;
 				if (bFlipHorizontally)
 				{
-					Buffer.StaticMeshVertexBuffer.SetVertexUV(0, 0, {0, 1});
-					Buffer.StaticMeshVertexBuffer.SetVertexUV(1, 0, {0, 0});
-					Buffer.StaticMeshVertexBuffer.SetVertexUV(2, 0, {1, 0});
-					Buffer.StaticMeshVertexBuffer.SetVertexUV(3, 0, {1, 1});
+					VertexBuildData[0].UVs[0] = {0, 1};
+					VertexBuildData[1].UVs[0] = {0, 0};
+					VertexBuildData[2].UVs[0] = {1, 0};
+					VertexBuildData[3].UVs[0] = {1, 1};
 				}
 				else
 				{
-					Buffer.StaticMeshVertexBuffer.SetVertexUV(0, 0, {1, 1});
-					Buffer.StaticMeshVertexBuffer.SetVertexUV(1, 0, {1, 0});
-					Buffer.StaticMeshVertexBuffer.SetVertexUV(2, 0, {0, 0});
-					Buffer.StaticMeshVertexBuffer.SetVertexUV(3, 0, {0, 1});
+					VertexBuildData[0].UVs[0] = {1, 1};
+					VertexBuildData[1].UVs[0] = {1, 0};
+					VertexBuildData[2].UVs[0] = {0, 0};
+					VertexBuildData[3].UVs[0] = {0, 1};
 				}
 
 				const bool bSRGB = true;
-				Buffer.ColorVertexBuffer.Init(NumVertices);
-				Buffer.ColorVertexBuffer.VertexColor(0) = InComponent->BoundsColor.ToFColor(bSRGB);
-				Buffer.ColorVertexBuffer.VertexColor(1) = InComponent->BoundsColor.ToFColor(bSRGB);
-				Buffer.ColorVertexBuffer.VertexColor(2) = InComponent->BoundsColor.ToFColor(bSRGB);
-				Buffer.ColorVertexBuffer.VertexColor(3) = InComponent->BoundsColor.ToFColor(bSRGB);
+				const FColor BoundsColor = InComponent->BoundsColor.ToFColor(bSRGB);
+
+				VertexBuildData[0].Color = BoundsColor;
+				VertexBuildData[1].Color = BoundsColor;
+				VertexBuildData[2].Color = BoundsColor;
+				VertexBuildData[3].Color = BoundsColor;
+
+				// MakeConstMeshBuildVertexView adds 8 (MAX_STATIC_TEXCOORDS) UVs instead of just the 1 we want.
+				const FConstMeshBuildVertexView VertexView
+				{
+					.Position = MakeConstStridedView(VertexBuildData, &FStaticMeshBuildVertex::Position),
+					.TangentX = MakeConstStridedView(VertexBuildData, &FStaticMeshBuildVertex::TangentX),
+					.TangentY = MakeConstStridedView(VertexBuildData, &FStaticMeshBuildVertex::TangentY),
+					.TangentZ = MakeConstStridedView(VertexBuildData, &FStaticMeshBuildVertex::TangentZ),
+					.UVs = { MakeConstStridedView(sizeof(FStaticMeshBuildVertex), &VertexBuildData[0].UVs[0], NumVertices) },
+					.Color = MakeConstStridedView(VertexBuildData, &FStaticMeshBuildVertex::Color),
+				};
+
+				Buffer.StaticMeshVertexBuffer.Init(VertexView);
+				Buffer.PositionVertexBuffer.Init(VertexView);
+				Buffer.ColorVertexBuffer.Init(VertexView);
 
 				Buffer.IndexBuffer.Indices = {0, 1, 2, 0, 2, 3};
 				Buffer.TriangleCount = NumTriangles;
