USD: Fix name collisions when exporting level sequences.

diff --git a/Engine/Plugins/Importers/USDImporter/Source/USDExporter/Private/LevelSequenceExporterUSD.cpp b/Engine/Plugins/Importers/USDImporter/Source/USDExporter/Private/LevelSequenceExporterUSD.cpp
--- a/Engine/Plugins/Importers/USDImporter/Source/USDExporter/Private/LevelSequenceExporterUSD.cpp
+++ b/Engine/Plugins/Importers/USDImporter/Source/USDExporter/Private/LevelSequenceExporterUSD.cpp
@@ -250,19 +250,6 @@ namespace UE::LevelSequenceExporterUSD::Private
 					*Guid.ToString()
 				);
 
-				if (AActor* SpawnedActor = Cast<AActor>(Object))
-				{
-					// Rename the spawn to a unique name or else in case of name collisions they will overwrite each other when writing
-					// animation data. The level exporter will rename actors to unique prims by itself though.
-					FString NewLabel = UsdUnreal::ObjectUtils::GetUniqueName(SpawnedActor->GetActorLabel(), UsedActorLabels);
-					if (NewLabel != SpawnedActor->GetActorLabel())
-					{
-						const bool bMarkDirty = false;
-						SpawnedActor->SetActorLabel(NewLabel, bMarkDirty);
-					}
-					UsedActorLabels.Add(NewLabel);
-				}
-
 				ExistingInstancesForGuid.Add(Object);
 				SpawnableIndices.Add(InstanceKey, ExistingInstancesForGuid.Num() - 1);
 			}
@@ -894,7 +881,9 @@ namespace UE::LevelSequenceExporterUSD::Private
 
 				UObject* BoundObject = nullptr;
 
-				// We need to check for custom spawnables here as well
+				// We need to check for custom spawnables here as well.
+				// Note: Now all sequencer bindings are possessables, even the old spawnables. This is why we loop
+				// over all binding references here, and will attempt to use the SpawnRegister even for possessables
 				if (const FMovieSceneBindingReferences* BindingReferences = MovieSceneSequence.GetBindingReferences())
 				{
 					int32 BindingIndex = 0;
diff --git a/Engine/Plugins/Runtime/USDCore/Source/USDUtilities/Private/USDConversionUtils.cpp b/Engine/Plugins/Runtime/USDCore/Source/USDUtilities/Private/USDConversionUtils.cpp
--- a/Engine/Plugins/Runtime/USDCore/Source/USDUtilities/Private/USDConversionUtils.cpp
+++ b/Engine/Plugins/Runtime/USDCore/Source/USDUtilities/Private/USDConversionUtils.cpp
@@ -106,6 +106,15 @@ static FAutoConsoleVariableRef CVarParseUVSetsFromFloat2Primvars(
 		 "to be parsed as UV sets.")
 );
 
+static bool GCheapUniquePrimPathGeneration = false;
+static FAutoConsoleVariableRef CVarCheapUniquePrimPathGeneration(
+	TEXT("USD.CheapUniquePrimPathGeneration"),
+	GCheapUniquePrimPathGeneration,
+	TEXT(
+		"When exporting Levels and LevelSequences, we'll by default use a mechanism of ensuring unique prim paths that guarantees a unique path for each UObject. It can be somewhat expensive depending on the use-case, so if you have other ways of ensuring actor labels are unique, you can set this to true to use another method of producing unique prim paths that is much faster, but can't handle some kinds of actor label collisions."
+	)
+);
+
 namespace USDConversionUtilsImpl
 {
 	/** Show some warnings if the UVSet primvars show some unsupported/problematic behavior */
@@ -628,7 +637,79 @@ FString UsdUtils::GetPrimPathForObject(const UObject* ActorOrComponent, const FS
 #if WITH_EDITOR
 	if (Component == Owner->GetRootComponent())
 	{
-		Path = Owner->GetActorLabel();
+		UObject* OwnerOuter = Owner->GetOuter();
+		FString OwnerLabel = Owner->GetActorLabel();
+
+		if (GCheapUniquePrimPathGeneration)
+		{
+			// This guarantees uniqueness only if all actors that have the same label also have the
+			// same FName text part (i.e. If we had a Directional Light named "Foo" and a StaticMeshActor
+			// named "Foo", their IDNames could end up being "DirectionalLight_2" and "StaticMeshActor_2",
+			// so this method would have generated the prim name "Foo_2" for both...)
+			Path = Owner->GetActorLabel() + TEXT("_") + LexToString(Owner->GetFName().GetNumber());
+		}
+		else
+		{
+			TArray<UObject*> SiblingActors;
+			bool bIncludeNestedObjects = false;
+			GetObjectsWithOuter(OwnerOuter, SiblingActors, bIncludeNestedObjects);
+
+			TSet<FString> SeenLabels;
+			SeenLabels.Reserve(SiblingActors.Num());
+
+			TArray<UObject*> SiblingsWithSameLabel;
+			for (UObject* Sibling : SiblingActors)
+			{
+				if (AActor* ActorSibling = Cast<AActor>(Sibling))
+				{
+					FString SiblingLabel = ActorSibling->GetActorLabel();
+					SeenLabels.Add(SiblingLabel);
+
+					if (ActorSibling == Owner || SiblingLabel == OwnerLabel)
+					{
+						SiblingsWithSameLabel.Add(Sibling);
+					}
+				}
+			}
+
+			// Sorting is important because we'll call this from e.g. LevelSequence export, and it should
+			// match the unique names that were generated on the Level export too
+			SiblingsWithSameLabel.Sort(
+				[](const UObject& LHS, const UObject& RHS)
+				{
+					return LHS.GetName() < RHS.GetName();
+				}
+			);
+
+			int32 Index = SiblingsWithSameLabel.IndexOfByKey(Owner);
+			if (Index == 0)
+			{
+				Path = Owner->GetActorLabel();
+			}
+			else
+			{
+				Path = Owner->GetActorLabel();
+
+				// Imagine we have the sibling actors with labels "Cube", "Cube", "Cube" and "Cube_0". In here suppose
+				// we're trying to come up with a prim name for the second of the "Cube"s. Normally we'd come up with "Cube_0",
+				// but there's already an actor with this label, so we can't use it. Unfortunately though, we can't
+				// just increment our index and use "Cube_1" either: When we sanitize the third "Cube" we'd end up also
+				// trying to name it "Cube_1" (remember, we don't keep any "state" between calls to this function)
+				//
+				// We also don't want to just add the "Cube_0" actor to the same list of name collisions and handle it in
+				// the same "group" as other "Cube"s because we want to preserve the user-set label if possible
+				// (i.e. we don't want one of the "Cube"s to end up exported with the previously existing name "Cube_0").
+				//
+				// The solution used here is to not increment our index, but to make sure we always add *a new* trailing
+				// suffix with it. That way, the labels for the actors in the example will end up being, respectively:
+				// "Cube", "Cube_0_0", "Cube_1", "Cube_0". It looks a bit goofy, but we don't need to preserve any state
+				// or global "used prim names" set anywhere, and it preserves "Cube_0" and even a "Cube" label
+				do
+				{
+					Path += TEXT("_") + LexToString(Index - 1);
+				} while (SeenLabels.Contains(Path));
+			}
+		}
 	}
 	else
 #endif	  // WITH_EDITOR
