USD: Add project setting option to enable/disable the diagnostic delegate

diff --git a/Engine/Plugins/Runtime/USDCore/Source/USDClasses/Public/USDProjectSettings.h b/Engine/Plugins/Runtime/USDCore/Source/USDClasses/Public/USDProjectSettings.h
--- a/Engine/Plugins/Runtime/USDCore/Source/USDClasses/Public/USDProjectSettings.h
+++ b/Engine/Plugins/Runtime/USDCore/Source/USDClasses/Public/USDProjectSettings.h
@@ -64,6 +64,10 @@ public:
 	UPROPERTY(config, EditAnywhere, Category = USD)
 	TArray<FName> AdditionalMaterialPurposes;
 
+	// Whether to show on the output log messages, warnings and errors reported directly by the USD SDK
+	UPROPERTY(config, EditAnywhere, Category = USD)
+	bool bLogUsdSdkErrors = true;
+
 	/**
 	 * USD Asset Cache to use for USD Stage Actors that don't have any asset cache specified.
 	 * Leave this empty to have each stage actor generate it's on transient cache instead.
diff --git a/Engine/Plugins/Runtime/USDCore/Source/UnrealUSDWrapper/Private/UnrealUSDWrapper.cpp b/Engine/Plugins/Runtime/USDCore/Source/UnrealUSDWrapper/Private/UnrealUSDWrapper.cpp
--- a/Engine/Plugins/Runtime/USDCore/Source/UnrealUSDWrapper/Private/UnrealUSDWrapper.cpp
+++ b/Engine/Plugins/Runtime/USDCore/Source/UnrealUSDWrapper/Private/UnrealUSDWrapper.cpp
@@ -1100,12 +1100,47 @@ void UnrealUSDWrapper::ClearDiagnosticDelegate()
 
 class FUnrealUSDWrapperModule : public IUnrealUSDWrapperModule
 {
+	FDelegateHandle OnProjectSettingsChangedHandle;
+
 public:
 	virtual void StartupModule() override
 	{
 #if USE_USD_SDK
 		LLM_SCOPE_BYTAG(Usd);
 
+		UUsdProjectSettings* ProjectSettings = GetMutableDefault<UUsdProjectSettings>();
+		if (ProjectSettings->bLogUsdSdkErrors)
+		{
+			UnrealUSDWrapper::SetupDiagnosticDelegate();
+		}
+
+#if WITH_EDITOR
+		OnProjectSettingsChangedHandle = ProjectSettings->OnSettingChanged().AddLambda(
+			[](UObject* SettingsObject, struct FPropertyChangedEvent& PropertyChangedEvent)
+			{
+				UUsdProjectSettings* UsdSettings = Cast<UUsdProjectSettings>(SettingsObject);
+				if (!UsdSettings)
+				{
+					return;
+				}
+
+				if (PropertyChangedEvent.GetPropertyName() == GET_MEMBER_NAME_CHECKED(UUsdProjectSettings, bLogUsdSdkErrors))
+				{
+					if (UsdSettings->bLogUsdSdkErrors)
+					{
+						UnrealUSDWrapper::SetupDiagnosticDelegate();
+						UE_LOG(LogUsd, Log, TEXT("Logging USD messages"));
+					}
+					else
+					{
+						UnrealUSDWrapper::ClearDiagnosticDelegate();
+						UE_LOG(LogUsd, Log, TEXT("Not logging USD messages"));
+					}
+				}
+			}
+		);
+#endif	  // WITH_EDITOR
+
 		// Path to USD base plugins
 		FString UsdPluginsPath = FPaths::Combine(TEXT(".."), TEXT("ThirdParty"), TEXT("USD"), TEXT("UsdResources"));
 		UsdPluginsPath = FPaths::ConvertRelativePathToFull(UsdPluginsPath);
@@ -1157,7 +1192,7 @@ public:
 		// Combine our current plugins with any additional USD plugins the user may have set.
 		TArray<FString> PluginDirectories;
 		PluginDirectories.Add(UsdPluginsPath);
-		for (const FDirectoryPath& Directory : GetDefault<UUsdProjectSettings>()->AdditionalPluginDirectories)
+		for (const FDirectoryPath& Directory : ProjectSettings->AdditionalPluginDirectories)
 		{
 			if (!Directory.Path.IsEmpty())
 			{
@@ -1193,13 +1228,12 @@ public:
 		}
 
 		// Set the default search path for USD's default resolver using any path specified in the settings.
-		const TArray<FDirectoryPath> SearchPath = GetDefault<UUsdProjectSettings>()->DefaultResolverSearchPath;
+		const TArray<FDirectoryPath> SearchPath = ProjectSettings->DefaultResolverSearchPath;
 		UnrealUSDWrapper::SetDefaultResolverDefaultSearchPath(SearchPath);
 
 #endif	  // USE_USD_SDK
 
 		FUsdMemoryManager::Initialize();
-		UnrealUSDWrapper::SetupDiagnosticDelegate();
 
 #if WITH_EDITOR
 		// Update the supported filetype filters for reference/payload picker dialogs
@@ -1228,6 +1262,18 @@ public:
 
 	virtual void ShutdownModule() override
 	{
+#if WITH_EDITOR
+		// We can't query default objects during engine exit
+		if (UObjectInitialized())
+		{
+			UUsdProjectSettings* ProjectSettings = GetMutableDefault<UUsdProjectSettings>();
+			if (ProjectSettings)
+			{
+				ProjectSettings->OnSettingChanged().Remove(OnProjectSettingsChangedHandle);
+			}
+		}
+#endif	  // WITH_EDITOR
+
 		UnrealUSDWrapper::ClearDiagnosticDelegate();
 		FUsdMemoryManager::Shutdown();
 	}
diff --git a/Engine/Plugins/Runtime/USDCore/Source/UnrealUSDWrapper/UnrealUSDWrapper.Build.cs b/Engine/Plugins/Runtime/USDCore/Source/UnrealUSDWrapper/UnrealUSDWrapper.Build.cs
--- a/Engine/Plugins/Runtime/USDCore/Source/UnrealUSDWrapper/UnrealUSDWrapper.Build.cs
+++ b/Engine/Plugins/Runtime/USDCore/Source/UnrealUSDWrapper/UnrealUSDWrapper.Build.cs
@@ -30,6 +30,16 @@ namespace UnrealBuildTool.Rules
 				}
 			);
 
+			if (Target.bBuildEditor)
+			{
+				PrivateDependencyModuleNames.AddRange(
+					new string[]
+					{
+						"DeveloperSettings", // So that we can un/register the diagnostic delegate when the project settings change
+					}
+				);
+			}
+
 			// Temporarily disabled runtime USD support until Mac and Linux dynamic linking issues are resolved
 			if (EnableUsdSdk(Target) && (Target.Type == TargetType.Editor || Target.Platform == UnrealTargetPlatform.Win64))
 			{
