USD: Allow setting USD project settings via Python/Utility Blueprints.

diff --git a/Engine/Plugins/Runtime/USDCore/Source/USDClasses/Private/USDProjectSettings.cpp b/Engine/Plugins/Runtime/USDCore/Source/USDClasses/Private/USDProjectSettings.cpp
new file mode 100644
--- /dev/null
+++ b/Engine/Plugins/Runtime/USDCore/Source/USDClasses/Private/USDProjectSettings.cpp
@@ -0,0 +1,13 @@
+// Copyright Epic Games, Inc. All Rights Reserved.
+
+#include "USDProjectSettings.h"
+
+#if WITH_EDITOR
+void UUsdProjectSettings::PostEditChangeProperty(FPropertyChangedEvent& PropertyChangedEvent)
+{
+	UDeveloperSettings::PostEditChangeProperty(PropertyChangedEvent);
+
+	// Allows settings to be applied even if the properties were set via utility blueprints/Python
+	SaveConfig();
+}
+#endif	  // WITH_EDITOR
diff --git a/Engine/Plugins/Runtime/USDCore/Source/USDClasses/Public/USDProjectSettings.h b/Engine/Plugins/Runtime/USDCore/Source/USDClasses/Public/USDProjectSettings.h
--- a/Engine/Plugins/Runtime/USDCore/Source/USDClasses/Public/USDProjectSettings.h
+++ b/Engine/Plugins/Runtime/USDCore/Source/USDClasses/Public/USDProjectSettings.h
@@ -27,7 +27,7 @@ enum class EUsdEditInInstanceBehavior : uint8
 // USDCore and defaultconfig here so this ends up at DefaultUSDCore.ini in the editor, and is sent to the
 // packaged game as well
 
-UCLASS(config = USDCore, defaultconfig, meta = (DisplayName = USDCore), MinimalAPI)
+UCLASS(BlueprintType, config = USDCore, defaultconfig, meta = (DisplayName = USDCore), MinimalAPI)
 class UUsdProjectSettings : public UDeveloperSettings
 {
 	GENERATED_BODY()
@@ -296,4 +296,9 @@ public:
 	UPROPERTY(config, EditAnywhere, Category = "USD|Reference Materials|DisplayColor", meta = (AllowedClasses = "/Script/Engine.MaterialInterface"))
 	FSoftObjectPath ReferenceDisplayColorAndOpacityTwoSidedMaterial = FSoftObjectPath{
 		TEXT("/USDCore/Materials/DisplayColorAndOpacityTwoSided.DisplayColorAndOpacityTwoSided")};
+
+public:
+#if WITH_EDITOR
+	virtual void PostEditChangeProperty(FPropertyChangedEvent& PropertyChangedEvent) override;
+#endif
 };
