USD: Write extents when converting assets to USD.

diff --git a/Engine/Plugins/Runtime/USDCore/Source/USDUtilities/Private/USDGeomMeshConversion.cpp b/Engine/Plugins/Runtime/USDCore/Source/USDUtilities/Private/USDGeomMeshConversion.cpp
--- a/Engine/Plugins/Runtime/USDCore/Source/USDUtilities/Private/USDGeomMeshConversion.cpp
+++ b/Engine/Plugins/Runtime/USDCore/Source/USDUtilities/Private/USDGeomMeshConversion.cpp
@@ -3902,6 +3902,8 @@ bool UnrealToUsd::ConvertStaticMesh(
 	// here, so would our referencer and we wouldn't be able to put a transform on it
 	UsdPrim = Stage->DefinePrim(UsdPrim.GetPath(), UnrealToUsd::ConvertToken(bExportMultipleLODs ? TEXT("Xform") : TEXT("Mesh")).Get());
 
+	TUsdStore<pxr::VtArray<pxr::GfVec3f>> USDBounds = UnrealToUsd::ConvertBounds(StageInfo, StaticMesh->GetBoundingBox());
+
 #if WITH_EDITOR
 	const bool bExportNaniteDataAsSourceData = GExportNaniteSourceMeshData && StaticMesh->IsNaniteEnabled()
 											   && StaticMesh->IsHiResMeshDescriptionValid();
@@ -3949,6 +3951,15 @@ bool UnrealToUsd::ConvertStaticMesh(
 			MaterialPrim = MaterialStage->OverridePrim(UsdPrim.GetPath());
 		}
 
+		// Export extents
+		if (USDBounds.Get().size() > 0)
+		{
+			if (pxr::UsdAttribute Attr = TargetMesh.CreateExtentAttr())
+			{
+				Attr.Set(USDBounds.Get());
+			}
+		}
+
 #if WITH_EDITOR
 		// Try exporting source data
 		//  - Reference: FFbxExporter::ExportStaticMeshToFBX
@@ -4543,6 +4554,7 @@ bool UnrealToUsd::ConvertGeometryCache(const UGeometryCache* GeometryCache, pxr:
 
 	pxr::UsdGeomMesh TargetMesh{UsdPrim};
 	pxr::UsdPrim MaterialPrim = MaterialStage->OverridePrim(UsdPrim.GetPath());
+	pxr::UsdAttribute ExtentsAttr = TargetMesh ? TargetMesh.CreateExtentAttr() : pxr::UsdAttribute{};
 
 	UsdGeometryCacheImpl::FGeometryCacheExportContext ExportContext(*GeometryCache);
 	const int32 StartFrame = GeometryCache->GetStartFrame();
@@ -4565,6 +4577,12 @@ bool UnrealToUsd::ConvertGeometryCache(const UGeometryCache* GeometryCache, pxr:
 			const float TimeCode = FrameIndex;
 			UsdGeometryCacheImpl::ConvertGeometryCacheMeshData(MeshData, TargetMesh, MaterialAssignments, TimeCode, MaterialPrim, ExportContext);
 		}
+
+		if (MeshData.BoundingBox.IsValid && ExtentsAttr)
+		{
+			TUsdStore<pxr::VtArray<pxr::GfVec3f>> USDBounds = UnrealToUsd::ConvertBounds(StageInfo, FBox{MeshData.BoundingBox});
+			ExtentsAttr.Set(USDBounds.Get(), pxr::UsdTimeCode(FrameIndex));
+		}
 	}
 
 	// Configure time metadata for the stage
diff --git a/Engine/Plugins/Runtime/USDCore/Source/USDUtilities/Private/USDSkeletalDataConversion.cpp b/Engine/Plugins/Runtime/USDCore/Source/USDUtilities/Private/USDSkeletalDataConversion.cpp
--- a/Engine/Plugins/Runtime/USDCore/Source/USDUtilities/Private/USDSkeletalDataConversion.cpp
+++ b/Engine/Plugins/Runtime/USDCore/Source/USDUtilities/Private/USDSkeletalDataConversion.cpp
@@ -3693,6 +3693,16 @@ bool UnrealToUsd::ConvertSkeletalMesh(
 		UnrealToUsd::ConvertSkeleton(SkeletalMesh->GetRefSkeleton(), SkelSkeleton);
 	}
 
+	// Export extents onto the SkelRoot
+	TUsdStore<pxr::VtArray<pxr::GfVec3f>> USDBounds = UnrealToUsd::ConvertBounds(StageInfo, SkeletalMesh->GetBounds().GetBox());
+	if (USDBounds.Get().size() > 0)
+	{
+		if (pxr::UsdAttribute Attr = SkelRoot.CreateExtentAttr())
+		{
+			Attr.Set(USDBounds.Get());
+		}
+	}
+
 	// Actual meshes
 	for (int32 LODIndex = LowestMeshLOD; LODIndex <= HighestMeshLOD; ++LODIndex)
 	{
@@ -3731,6 +3741,15 @@ bool UnrealToUsd::ConvertSkeletalMesh(
 		pxr::UsdPrim UsdLODPrim = Stage->DefinePrim(MeshPrimPath, UnrealToUsd::ConvertToken(TEXT("Mesh")).Get());
 		pxr::UsdGeomMesh UsdLODPrimGeomMesh{UsdLODPrim};
 
+		// Export extents onto the Mesh itself too (it's the same extent in our case as we always just have one mesh)
+		if (USDBounds.Get().size() > 0)
+		{
+			if (pxr::UsdAttribute Attr = UsdLODPrimGeomMesh.CreateExtentAttr())
+			{
+				Attr.Set(USDBounds.Get());
+			}
+		}
+
 		pxr::UsdPrim MaterialPrim = UsdLODPrim;
 		if (StageForMaterialAssignments)
 		{
@@ -3858,6 +3877,9 @@ bool UnrealToUsd::ConvertAnimSequence(UAnimSequence* AnimSequence, pxr::UsdPrim&
 	FScopedUsdAllocs UsdAllocs;
 	pxr::SdfChangeBlock ChangeBlock;
 
+	pxr::UsdSkelRoot ParentSkelRoot = pxr::UsdSkelRoot{UsdUtils::GetClosestParentSkelRoot(SkelAnimPrim)};
+	pxr::UsdAttribute ExtentsAttr = ParentSkelRoot ? ParentSkelRoot.CreateExtentAttr() : pxr::UsdAttribute{};
+
 	FUsdStageInfo StageInfo(SkelAnimPrim.GetStage());
 
 	// Blend shapes
@@ -3908,7 +3930,7 @@ bool UnrealToUsd::ConvertAnimSequence(UAnimSequence* AnimSequence, pxr::UsdPrim&
 		UnrealToUsd::ConvertJointsAttribute(RefSkeleton, JointsAttr);
 	}
 
-	// Translations, Rotations & Scales
+	// Translations, Rotations, Scales & Extents
 	{
 		pxr::UsdAttribute TranslationsAttr = UsdSkelAnim.CreateTranslationsAttr();
 		pxr::UsdAttribute RotationsAttr = UsdSkelAnim.CreateRotationsAttr();
@@ -3954,6 +3976,13 @@ bool UnrealToUsd::ConvertAnimSequence(UAnimSequence* AnimSequence, pxr::UsdPrim&
 			TranslationsAttr.Set(Translations, pxr::UsdTimeCode(TimeCode));
 			RotationsAttr.Set(Rotations, pxr::UsdTimeCode(TimeCode));
 			ScalesAttr.Set(Scales, pxr::UsdTimeCode(TimeCode));
+
+			FBox Bounds = DebugSkelMeshComponent->CalcBounds(FTransform::Identity).GetBox();
+			if (Bounds.IsValid && ExtentsAttr)
+			{
+				TUsdStore<pxr::VtArray<pxr::GfVec3f>> USDBounds = UnrealToUsd::ConvertBounds(StageInfo, Bounds);
+				ExtentsAttr.Set(USDBounds.Get(), pxr::UsdTimeCode(TimeCode));
+			}
 		}
 
 		// Actively delete it or else it will remain visible on the viewport
diff --git a/Engine/Plugins/Runtime/USDCore/Source/USDUtilities/Private/USDTypesConversion.cpp b/Engine/Plugins/Runtime/USDCore/Source/USDUtilities/Private/USDTypesConversion.cpp
--- a/Engine/Plugins/Runtime/USDCore/Source/USDUtilities/Private/USDTypesConversion.cpp
+++ b/Engine/Plugins/Runtime/USDCore/Source/USDUtilities/Private/USDTypesConversion.cpp
@@ -18,6 +18,7 @@
 #include "pxr/base/gf/vec2f.h"
 #include "pxr/base/gf/vec3f.h"
 #include "pxr/base/gf/vec4f.h"
+#include "pxr/base/vt/array.h"
 #include "pxr/usd/sdf/path.h"
 #include "pxr/usd/usdGeom/tokens.h"
 #include "USDIncludesEnd.h"
@@ -573,6 +574,34 @@ namespace UnrealToUsd
 
 		return Value;
 	}
+
+	TUsdStore<pxr::VtArray<pxr::GfVec3f>> ConvertBounds(const FUsdStageInfo& StageInfo, FBox Bounds)
+	{
+		TUsdStore<pxr::VtArray<pxr::GfVec3f>> Result;
+		if (!Bounds.IsValid)
+		{
+			return Result;
+		}
+
+		pxr::GfVec3f ConvertedMin = UnrealToUsd::ConvertVectorFloat(StageInfo, Bounds.Min);
+		pxr::GfVec3f ConvertedMax = UnrealToUsd::ConvertVectorFloat(StageInfo, Bounds.Max);
+
+		// Because of the up axis conversion, we may have e.g. ConvertedMin.y > ConvertedMax.y, so here we
+		// fix that up
+		pxr::GfVec3f ActualMin = {
+			FMath::Min(ConvertedMin[0], ConvertedMax[0]),
+			FMath::Min(ConvertedMin[1], ConvertedMax[1]),
+			FMath::Min(ConvertedMin[2], ConvertedMax[2])};
+		pxr::GfVec3f ActualMax = {
+			FMath::Max(ConvertedMin[0], ConvertedMax[0]),
+			FMath::Max(ConvertedMin[1], ConvertedMax[1]),
+			FMath::Max(ConvertedMin[2], ConvertedMax[2])};
+
+		pxr::VtArray<pxr::GfVec3f>& ResultInner = Result.Get();
+		ResultInner = {ActualMin, ActualMax};
+
+		return Result;
+	}
 }	 // namespace UnrealToUsd
 
 #endif	  // #if USE_USD_SDK
diff --git a/Engine/Plugins/Runtime/USDCore/Source/USDUtilities/Public/USDTypesConversion.h b/Engine/Plugins/Runtime/USDCore/Source/USDUtilities/Public/USDTypesConversion.h
--- a/Engine/Plugins/Runtime/USDCore/Source/USDUtilities/Public/USDTypesConversion.h
+++ b/Engine/Plugins/Runtime/USDCore/Source/USDUtilities/Public/USDTypesConversion.h
@@ -40,6 +40,9 @@ PXR_NAMESPACE_OPEN_SCOPE
 	class TfRefPtr;
 
 	using UsdStageRefPtr = TfRefPtr<UsdStage>;
+
+	template <typename T>
+	class VtArray;
 PXR_NAMESPACE_CLOSE_SCOPE
 
 struct USDUTILITIES_API FUsdStageInfo
@@ -154,6 +157,8 @@ namespace UnrealToUsd
 
 	/** Returns a distance in USD units (depends on metersPerUnit) */
 	USDUTILITIES_API float ConvertDistance(const FUsdStageInfo& StageInfo, float InValue);
+
+	USDUTILITIES_API TUsdStore<pxr::VtArray<pxr::GfVec3f>> ConvertBounds(const FUsdStageInfo& StageInfo, FBox Bounds);
 }	 // namespace UnrealToUsd
 
 namespace UsdUtils
