[USD]
* Fixed so usd module is not pulled in in game/client/server targets since they now can be built non-monolithic (and usd was enabled for all non-monolithic targets before)

diff --git a/Engine/Plugins/Runtime/USDCore/Source/UnrealUSDWrapper/UnrealUSDWrapper.Build.cs b/Engine/Plugins/Runtime/USDCore/Source/UnrealUSDWrapper/UnrealUSDWrapper.Build.cs
--- a/Engine/Plugins/Runtime/USDCore/Source/UnrealUSDWrapper/UnrealUSDWrapper.Build.cs
+++ b/Engine/Plugins/Runtime/USDCore/Source/UnrealUSDWrapper/UnrealUSDWrapper.Build.cs
@@ -224,7 +224,8 @@ namespace UnrealBuildTool.Rules
 			// Note that by that point everything will be using the USD-compatible ANSI allocators anyway, so our overrides in USDMemory.h are also disabled, as they're unnecessary.
 			// Also note that we're forced to use dynamic linking for monolithic targets mainly because static linking the USD libraries disables support for user USD plugins, and secondly
 			// because those static libraries would need to be linked with the --whole-archive argument, and there is currently no standard way of doing that in UE.
-			if (bEnableUsdSdk && Target.LinkType == TargetLinkType.Monolithic && !Target.GlobalDefinitions.Contains("FORCE_ANSI_ALLOCATOR=1") && !Target.GlobalDefinitions.Contains("UE_USE_MALLOC_FILL_BYTES=0"))
+			bool bIsGameOrMonolithic = Target.LinkType == TargetLinkType.Monolithic || Target.Type == TargetType.Client || Target.Type == TargetType.Server || Target.Type == TargetType.Game;
+			if (bEnableUsdSdk && bIsGameOrMonolithic && !Target.GlobalDefinitions.Contains("FORCE_ANSI_ALLOCATOR=1") && !Target.GlobalDefinitions.Contains("UE_USE_MALLOC_FILL_BYTES=0"))
 			{
 				PublicDefinitions.Add("USD_FORCE_DISABLED=1");
 				bEnableUsdSdk = false;
