USD: Allow exporting static meshes with multiple LODs and Nanite override enabled at the same time.

Also allow importing meshes with multiple LODs and Nanite override enabled.

#9410

diff --git a/Engine/Plugins/Importers/USDImporter/Source/USDSchemas/Private/USDGeomMeshTranslator.cpp b/Engine/Plugins/Importers/USDImporter/Source/USDSchemas/Private/USDGeomMeshTranslator.cpp
--- a/Engine/Plugins/Importers/USDImporter/Source/USDSchemas/Private/USDGeomMeshTranslator.cpp
+++ b/Engine/Plugins/Importers/USDImporter/Source/USDSchemas/Private/USDGeomMeshTranslator.cpp
@@ -185,23 +185,6 @@ namespace UsdGeomMeshTranslatorImpl
 			}
 		}
 
-		// Don't enable Nanite if we have more than one LOD. This means the Mesh came from the LOD variant set setup, and
-		// we're considering the LOD setup "stronger" than the Nanite override: If you have all that LOD variant set situation you
-		// likely don't want Nanite for one of the LOD meshes anyway, as that doesn't really make any sense.
-		// If the user wants to have Nanite within the variant set all they would otherwise need is to name the variant set something
-		// else other than LOD.
-		if (LODIndexToMeshDescription.Num() > 1)
-		{
-			UE_LOG(
-				LogUsd,
-				Warning,
-				TEXT("Not enabling Nanite for mesh generated for prim '%s' as it has more than one generated LOD (and so came from a LOD variant set "
-					 "setup)"),
-				*PrimPath.GetString()
-			);
-			return false;
-		}
-
 		if (Context.UsdInfoCache)
 		{
 			TOptional<uint64> SubtreeSectionCount = Context.UsdInfoCache->GetSubtreeMaterialSlotCount(PrimPath);
diff --git a/Engine/Plugins/Runtime/USDCore/Source/USDUtilities/Private/USDGeomMeshConversion.cpp b/Engine/Plugins/Runtime/USDCore/Source/USDUtilities/Private/USDGeomMeshConversion.cpp
--- a/Engine/Plugins/Runtime/USDCore/Source/USDUtilities/Private/USDGeomMeshConversion.cpp
+++ b/Engine/Plugins/Runtime/USDCore/Source/USDUtilities/Private/USDGeomMeshConversion.cpp
@@ -3815,35 +3815,7 @@ bool UnrealToUsd::ConvertStaticMesh(
 	// Make sure it's at least 1 LOD level
 	NumLODs = FMath::Max(HighestMeshLOD - LowestMeshLOD + 1, 1);
 
-	// If exporting a Nanite mesh, just use the lowest LOD and write the unrealNanite override attribute,
-	// this way we can guarantee it will have Nanite when it's imported back.
-#if WITH_EDITOR
-	if (StaticMesh->IsNaniteEnabled())
-	{
-		if (NumLODs > 1)
-		{
-			UE_LOG(
-				LogUsd,
-				Log,
-				TEXT("Not exporting multiple LODs for mesh '%s' onto prim '%s' since the mesh has Nanite enabled: LOD '%d' will be used and the '%s' "
-					 "attribute will be written out instead"),
-				*StaticMesh->GetName(),
-				*UsdToUnreal::ConvertPath(UsdPrim.GetPrimPath()),
-				LowestMeshLOD,
-				*UsdToUnreal::ConvertToken(UnrealIdentifiers::UnrealNaniteOverride)
-			);
-		}
-
-		HighestMeshLOD = LowestMeshLOD;
-		NumLODs = 1;
-
-		if (pxr::UsdAttribute Attr = UsdPrim.CreateAttribute(UnrealIdentifiers::UnrealNaniteOverride, pxr::SdfValueTypeNames->Token))
-		{
-			Attr.Set(UnrealIdentifiers::UnrealNaniteOverrideEnable);
-			UsdUtils::NotifyIfOverriddenOpinion(Attr);
-		}
-	}
-#endif	  // WITH_EDITOR
+	const bool bNaniteEnabled = StaticMesh->IsNaniteEnabled();
 
 	pxr::UsdVariantSets VariantSets = UsdPrim.GetVariantSets();
 	pxr::UsdVariantSet VariantSet = VariantSets.GetVariantSet(UnrealIdentifiers::LOD);
@@ -3951,6 +3923,16 @@ bool UnrealToUsd::ConvertStaticMesh(
 			MaterialPrim = MaterialStage->OverridePrim(UsdPrim.GetPath());
 		}
 
+		// Enable Nanite if we have it
+		if (bNaniteEnabled)
+		{
+			if (pxr::UsdAttribute Attr = TargetMesh.GetPrim().CreateAttribute(UnrealIdentifiers::UnrealNaniteOverride, pxr::SdfValueTypeNames->Token))
+			{
+				Attr.Set(UnrealIdentifiers::UnrealNaniteOverrideEnable);
+				UsdUtils::NotifyIfOverriddenOpinion(Attr);
+			}
+		}
+
 		// Export extents
 		if (USDBounds.Get().size() > 0)
 		{
