USD: Fix non-primitive-shape Gprims within collapsing hierarchies producing empty MeshDescriptions.

diff --git a/Engine/Plugins/Runtime/USDCore/Source/USDUtilities/Private/USDGeomMeshConversion.cpp b/Engine/Plugins/Runtime/USDCore/Source/USDUtilities/Private/USDGeomMeshConversion.cpp
--- a/Engine/Plugins/Runtime/USDCore/Source/USDUtilities/Private/USDGeomMeshConversion.cpp
+++ b/Engine/Plugins/Runtime/USDCore/Source/USDUtilities/Private/USDGeomMeshConversion.cpp
@@ -967,8 +967,9 @@ namespace UE::UsdGeomMeshConversion::Private
 				bSuccess = UsdToUnreal::ConvertGeomMesh(Mesh, OutMeshDescription, OutMaterialAssignments, Options);
 			}
 		}
-		// Check for a cube/capsule/etc. ConvertGeomPrimitive will internally check for all specific types
-		else if (pxr::UsdGeomGprim Gprim = pxr::UsdGeomGprim{Prim})
+		// Check for primitive types directly (don't just check for Gprim here because if we get something else like a UsdVolVolume or
+		// UsdGeomBasisCurves ConvertGeomPrimitive will return false and we will dump our MeshDescription)
+		else if (Prim.IsA<pxr::UsdGeomCapsule>() || Prim.IsA<pxr::UsdGeomCone>() || Prim.IsA<pxr::UsdGeomCube>() || Prim.IsA<pxr::UsdGeomCylinder>() || Prim.IsA<pxr::UsdGeomPlane>() || Prim.IsA<pxr::UsdGeomSphere>())
 		{
 			bSuccess = UsdToUnreal::ConvertGeomPrimitive(Prim, OutMeshDescription, OutMaterialAssignments, Options);
 		}
