USD: Fix compile error for runtime due to trying to use the Nanite flag without the editor.

diff --git a/Engine/Plugins/Runtime/USDCore/Source/USDUtilities/Private/USDGeomMeshConversion.cpp b/Engine/Plugins/Runtime/USDCore/Source/USDUtilities/Private/USDGeomMeshConversion.cpp
--- a/Engine/Plugins/Runtime/USDCore/Source/USDUtilities/Private/USDGeomMeshConversion.cpp
+++ b/Engine/Plugins/Runtime/USDCore/Source/USDUtilities/Private/USDGeomMeshConversion.cpp
@@ -3816,7 +3816,11 @@ bool UnrealToUsd::ConvertStaticMesh(
 	// Make sure it's at least 1 LOD level
 	NumLODs = FMath::Max(HighestMeshLOD - LowestMeshLOD + 1, 1);
 
+#if WITH_EDITOR
 	const bool bNaniteEnabled = StaticMesh->IsNaniteEnabled();
+#else
+	const bool bNaniteEnabled = false;
+#endif	  // WITH_EDITOR
 
 	pxr::UsdVariantSets VariantSets = UsdPrim.GetVariantSets();
 	pxr::UsdVariantSet VariantSet = VariantSets.GetVariantSet(UnrealIdentifiers::LOD);
