USD: Added option to export skeletal animation as geometry cache.

diff --git a/Engine/Plugins/Importers/USDImporter/Source/USDExporter/Private/AnimSequenceExporterUSD.cpp b/Engine/Plugins/Importers/USDImporter/Source/USDExporter/Private/AnimSequenceExporterUSD.cpp
--- a/Engine/Plugins/Importers/USDImporter/Source/USDExporter/Private/AnimSequenceExporterUSD.cpp
+++ b/Engine/Plugins/Importers/USDImporter/Source/USDExporter/Private/AnimSequenceExporterUSD.cpp
@@ -132,8 +132,12 @@ bool UAnimSequenceExporterUSD::ExportBinary(
 	// Export preview mesh if needed
 	USkeletalMesh* SkeletalMesh = nullptr;
 	FString MeshAssetFile;
+	bool bExportAsSkeletal = true;
 	if (Options && Options->bExportPreviewMesh)
 	{
+		// To convert to non-skeletal, need to export the preview mesh in the first place
+		bExportAsSkeletal = !Options->PreviewMeshOptions.bConvertSkeletalToNonSkeletal;
+
 		SkeletalMesh = AnimSequence->GetPreviewMesh();
 		USkeleton* AnimSkeleton = SkeletalMesh ? SkeletalMesh->GetSkeleton() : nullptr;
 
@@ -196,7 +200,7 @@ bool UAnimSequenceExporterUSD::ExportBinary(
 	// Collect the target paths for our SkelAnimation prim and its SkelRoot, if any
 	UE::FSdfPath SkelRootPath;
 	UE::FSdfPath SkelAnimPath;
-	if (MeshAssetFile.IsEmpty())
+	if (MeshAssetFile.IsEmpty() || !bExportAsSkeletal)
 	{
 		SkelAnimPath = UE::FSdfPath::AbsoluteRootPath().AppendChild(*UsdUtils::SanitizeUsdIdentifier(*AnimSequence->GetName()));
 	}
@@ -295,34 +299,50 @@ bool UAnimSequenceExporterUSD::ExportBinary(
 	// Exported a SkeletalMesh prim elsewhere, create a SkelRoot containing this SkelAnimation prim
 	else
 	{
-		SkelRootPrim = AnimationStage.DefinePrim(SkelRootPath, TEXT("SkelRoot"));
-		if (!SkelRootPrim)
+		if (bExportAsSkeletal)
 		{
-			return false;
-		}
+			SkelRootPrim = AnimationStage.DefinePrim(SkelRootPath, TEXT("SkelRoot"));
+			if (!SkelRootPrim)
+			{
+				return false;
+			}
 
-		SkelAnimPrim = AnimationStage.DefinePrim(SkelAnimPath, TEXT("SkelAnimation"));
-		if (!SkelAnimPrim)
-		{
-			return false;
-		}
+			SkelAnimPrim = AnimationStage.DefinePrim(SkelAnimPath, TEXT("SkelAnimation"));
+			if (!SkelAnimPrim)
+			{
+				return false;
+			}
 
-		UE::FSdfPath SkeletonPath = SkelRootPath.AppendChild(UnrealIdentifiers::ExportedSkeletonPrimName);
-		UE::FUsdPrim SkeletonPrim = AnimationStage.DefinePrim(SkeletonPath, TEXT("Skeleton"));
-		if (!SkeletonPrim)
-		{
-			return false;
-		}
+			UE::FSdfPath SkeletonPath = SkelRootPath.AppendChild(UnrealIdentifiers::ExportedSkeletonPrimName);
+			UE::FUsdPrim SkeletonPrim = AnimationStage.DefinePrim(SkeletonPath, TEXT("Skeleton"));
+			if (!SkeletonPrim)
+			{
+				return false;
+			}
 
-		AnimationStage.SetDefaultPrim(SkelRootPrim);
+			AnimationStage.SetDefaultPrim(SkelRootPrim);
 
-		// Add a reference to the SkelRoot of the static mesh, which will compose in the Mesh and Skeleton prims
-		UsdUtils::AddReference(SkelRootPrim, *MeshAssetFile);
+			// Add a reference to the SkelRoot of the static mesh, which will compose in the Mesh and Skeleton prims
+			UsdUtils::AddReference(SkelRootPrim, *MeshAssetFile);
 
-		// We bind the animation directly to the Skeleton (and not the skel root) because binding it to the SkelRoot
-		// may lead to trouble when we're exporting nested SkeletalMeshComponents, as it will be inherited by
-		// all child Skeletons (even the ones that wouldn't otherwise receive any animation)
-		UsdUtils::BindAnimationSource(SkeletonPrim, SkelAnimPrim);
+			// We bind the animation directly to the Skeleton (and not the skel root) because binding it to the SkelRoot
+			// may lead to trouble when we're exporting nested SkeletalMeshComponents, as it will be inherited by
+			// all child Skeletons (even the ones that wouldn't otherwise receive any animation)
+			UsdUtils::BindAnimationSource(SkeletonPrim, SkelAnimPrim);
+		}
+		else
+		{
+			SkelAnimPrim = AnimationStage.DefinePrim(SkelAnimPath, TEXT("Mesh"));
+			if (!SkelAnimPrim)
+			{
+				return false;
+			}
+
+			AnimationStage.SetDefaultPrim(SkelAnimPrim);
+
+			// Add a reference to the static mesh prim on which the time samples will be authored
+			UsdUtils::AddReference(SkelAnimPrim, *MeshAssetFile);
+		}
 	}
 
 	// Configure stage metadata
@@ -340,7 +360,14 @@ bool UAnimSequenceExporterUSD::ExportBinary(
 		AnimationStage.SetTimeCodesPerSecond(AnimSequence->GetSamplingFrameRate().AsDecimal());
 	}
 
-	UnrealToUsd::ConvertAnimSequence(AnimSequence, SkelAnimPrim);
+	if (bExportAsSkeletal)
+	{
+		UnrealToUsd::ConvertAnimSequence(AnimSequence, SkelAnimPrim);
+	}
+	else
+	{
+		UnrealToUsd::ConvertAnimSequenceToAnimatedMesh(AnimSequence, SkeletalMesh, SkelAnimPrim);
+	}
 
 	if (Options->MetadataOptions.bExportAssetInfo)
 	{
diff --git a/Engine/Plugins/Importers/USDImporter/Source/USDExporter/Private/AnimSequenceExporterUSDOptions.cpp b/Engine/Plugins/Importers/USDImporter/Source/USDExporter/Private/AnimSequenceExporterUSDOptions.cpp
--- a/Engine/Plugins/Importers/USDImporter/Source/USDExporter/Private/AnimSequenceExporterUSDOptions.cpp
+++ b/Engine/Plugins/Importers/USDImporter/Source/USDExporter/Private/AnimSequenceExporterUSDOptions.cpp
@@ -21,4 +21,8 @@ void UsdUtils::HashForAnimSequenceExport(const UAnimSequenceExporterUSDOptions&
 	UsdUtils::HashForExport(Options.StageOptions, HashToUpdate);
 	UsdUtils::HashForExport(Options.MetadataOptions, HashToUpdate);
 	HashToUpdate.Update(reinterpret_cast<const uint8*>(&Options.bExportPreviewMesh), sizeof(Options.bExportPreviewMesh));
+	if (Options.bExportPreviewMesh)
+	{
+		HashToUpdate.Update(reinterpret_cast<const uint8*>(&Options.PreviewMeshOptions.bConvertSkeletalToNonSkeletal), sizeof(Options.PreviewMeshOptions.bConvertSkeletalToNonSkeletal));
+	}
 }
diff --git a/Engine/Plugins/Importers/USDImporter/Source/USDExporter/Private/LevelSequenceExporterUSD.cpp b/Engine/Plugins/Importers/USDImporter/Source/USDExporter/Private/LevelSequenceExporterUSD.cpp
--- a/Engine/Plugins/Importers/USDImporter/Source/USDExporter/Private/LevelSequenceExporterUSD.cpp
+++ b/Engine/Plugins/Importers/USDImporter/Source/USDExporter/Private/LevelSequenceExporterUSD.cpp
@@ -9,6 +9,7 @@
 #include "USDConversionUtils.h"
 #include "USDErrorUtils.h"
 #include "USDExporterModule.h"
+#include "USDGeomMeshConversion.h"
 #include "USDLayerUtils.h"
 #include "USDLog.h"
 #include "USDObjectUtils.h"
@@ -1138,25 +1139,35 @@ namespace UE::LevelSequenceExporterUSD::Private
 					}
 				};
 
+			const bool bBakeAsSkeletal = !Context.ExportOptions->LevelExportOptions.AssetOptions.bConvertSkeletalToNonSkeletal;
 			TFunction<void(UnrealToUsd::FComponentBaker&)> GenerateSkeletalBaker =
-				[&UsdStage, &BoundComponent, &PrimPath](UnrealToUsd::FComponentBaker& InOutBaker)
+				[&UsdStage, &BoundComponent, &PrimPath, &Prim, bBakeAsSkeletal](UnrealToUsd::FComponentBaker& InOutBaker)
 				{
 					if (USkeletalMeshComponent* SkeletalBoundComponent = Cast<USkeletalMeshComponent>(BoundComponent))
 					{
-						UE::FUsdPrim SkelAnimPrim = UsdStage.DefinePrim(UE::FSdfPath{ *PrimPath }.AppendChild(TEXT("Anim")), TEXT("SkelAnimation"));
+						if (bBakeAsSkeletal)
+						{
+							UE::FUsdPrim SkelAnimPrim = UsdStage.DefinePrim(UE::FSdfPath{ *PrimPath }.AppendChild(TEXT("Anim")), TEXT("SkelAnimation"));
 
-						UE::FUsdPrim SkeletonPrim = UsdStage.DefinePrim(
-							UE::FSdfPath{ *PrimPath }.AppendChild(UnrealIdentifiers::ExportedSkeletonPrimName),
-							TEXT("Skeleton")
-						);
+							UE::FUsdPrim SkeletonPrim = UsdStage.DefinePrim(
+								UE::FSdfPath{ *PrimPath }.AppendChild(UnrealIdentifiers::ExportedSkeletonPrimName),
+								TEXT("Skeleton")
+							);
 
-						if (SkelAnimPrim && SkeletonPrim)
-						{
-							UnrealToUsd::CreateSkeletalAnimationBaker(SkeletonPrim, SkelAnimPrim, *SkeletalBoundComponent, InOutBaker);
+							if (SkelAnimPrim && SkeletonPrim)
+							{
+								UnrealToUsd::CreateSkeletalAnimationBaker(SkeletonPrim, SkelAnimPrim, *SkeletalBoundComponent, InOutBaker);
+							}
+							else
+							{
+								UE_LOG(LogUsd, Warning, TEXT("Failed to generate Skeleton or SkelAnimation prim when baking out SkelRoot '%s'"), *PrimPath);
+							}
 						}
 						else
 						{
-							UE_LOG(LogUsd, Warning, TEXT("Failed to generate Skeleton or SkelAnimation prim when baking out SkelRoot '%s'"), *PrimPath);
+							// Convert the prim for the skeletal mesh component from SkelRoot to Mesh
+							Prim.SetTypeName(TEXT("Mesh"));
+							UnrealToUsd::CreateSkeletalAnimationToMeshBaker(Prim, *SkeletalBoundComponent, InOutBaker);
 						}
 					}
 				};
diff --git a/Engine/Plugins/Importers/USDImporter/Source/USDExporter/Private/SkeletalMeshExporterUSD.cpp b/Engine/Plugins/Importers/USDImporter/Source/USDExporter/Private/SkeletalMeshExporterUSD.cpp
--- a/Engine/Plugins/Importers/USDImporter/Source/USDExporter/Private/SkeletalMeshExporterUSD.cpp
+++ b/Engine/Plugins/Importers/USDImporter/Source/USDExporter/Private/SkeletalMeshExporterUSD.cpp
@@ -273,7 +273,8 @@ bool USkeletalMeshExporterUsd::ExportBinary(
 
 	FString RootPrimPath = (TEXT("/") + UsdUtils::SanitizeUsdIdentifier(*SkeletalMesh->GetName()));
 
-	UE::FUsdPrim RootPrim = UsdStage.DefinePrim(UE::FSdfPath(*RootPrimPath), TEXT("SkelRoot"));
+	const bool bExportAsSkeletal = !Options->MeshAssetOptions.bConvertSkeletalToNonSkeletal;
+	UE::FUsdPrim RootPrim = UsdStage.DefinePrim(UE::FSdfPath(*RootPrimPath), bExportAsSkeletal ? TEXT("SkelRoot") : TEXT("Mesh"));
 	if (!RootPrim)
 	{
 		return false;
@@ -307,14 +308,26 @@ bool USkeletalMeshExporterUsd::ExportBinary(
 		AssetStage = UsdStage;
 	}
 
-	UnrealToUsd::ConvertSkeletalMesh(
-		SkeletalMesh,
-		RootPrim,
-		UsdUtils::GetDefaultTimeCode(),
-		&AssetStage,
-		Options->MeshAssetOptions.LowestMeshLOD,
-		Options->MeshAssetOptions.HighestMeshLOD
-	);
+	if (bExportAsSkeletal)
+	{
+		UnrealToUsd::ConvertSkeletalMesh(
+			SkeletalMesh,
+			RootPrim,
+			UsdUtils::GetDefaultTimeCode(),
+			&AssetStage,
+			Options->MeshAssetOptions.LowestMeshLOD,
+			Options->MeshAssetOptions.HighestMeshLOD
+		);
+	}
+	else
+	{
+		UnrealToUsd::ConvertSkeletalMeshToStaticMesh(
+			SkeletalMesh,
+			RootPrim,
+			UsdUtils::GetDefaultTimeCode(),
+			&AssetStage
+		);
+	}
 
 	if (UE::FUsdPrim AssetDefaultPrim = AssetStage.GetDefaultPrim())
 	{
diff --git a/Engine/Plugins/Importers/USDImporter/Source/USDExporter/Private/SkeletalMeshExporterUSDOptions.cpp b/Engine/Plugins/Importers/USDImporter/Source/USDExporter/Private/SkeletalMeshExporterUSDOptions.cpp
--- a/Engine/Plugins/Importers/USDImporter/Source/USDExporter/Private/SkeletalMeshExporterUSDOptions.cpp
+++ b/Engine/Plugins/Importers/USDImporter/Source/USDExporter/Private/SkeletalMeshExporterUSDOptions.cpp
@@ -16,5 +16,7 @@ void UsdUtils::HashForSkeletalMeshExport(const USkeletalMeshExporterUSDOptions&
 {
 	UsdUtils::HashForExport(Options.StageOptions, HashToUpdate);
 	UsdUtils::HashForMeshExport(Options.MeshAssetOptions, HashToUpdate);
+	// This option is only relevant to skeletal mesh export so it is not hashed as part of the mesh export
+	HashToUpdate.Update(reinterpret_cast<const uint8*>(&Options.MeshAssetOptions.bConvertSkeletalToNonSkeletal), sizeof(Options.MeshAssetOptions.bConvertSkeletalToNonSkeletal));
 	UsdUtils::HashForExport(Options.MetadataOptions, HashToUpdate);
 }
diff --git a/Engine/Plugins/Importers/USDImporter/Source/USDExporter/Public/USDAssetOptions.h b/Engine/Plugins/Importers/USDImporter/Source/USDExporter/Public/USDAssetOptions.h
--- a/Engine/Plugins/Importers/USDImporter/Source/USDExporter/Public/USDAssetOptions.h
+++ b/Engine/Plugins/Importers/USDImporter/Source/USDExporter/Public/USDAssetOptions.h
@@ -89,6 +89,10 @@ struct USDEXPORTER_API FUsdMeshAssetOptions
 	UPROPERTY(EditAnywhere, config, BlueprintReadWrite, Category = "Mesh options")
 	bool bExportStaticMeshSourceData = true;
 
+	/** Whether to convert skeletal data to non-skeletal, ie. skeletal mesh to static mesh, anim sequence to animated mesh */
+	UPROPERTY(EditAnywhere, config, BlueprintReadWrite, Category = "Mesh options")
+	bool bConvertSkeletalToNonSkeletal = false;
+
 	/** Lowest of the LOD indices to export static and skeletal meshes with (use 0 for full detail) */
 	UPROPERTY(EditAnywhere, config, BlueprintReadWrite, Category = "Mesh options", meta = (ClampMin = "0"))
 	int32 LowestMeshLOD = 0;
diff --git a/Engine/Plugins/Runtime/USDCore/Source/USDUtilities/Private/USDGeomMeshConversion.cpp b/Engine/Plugins/Runtime/USDCore/Source/USDUtilities/Private/USDGeomMeshConversion.cpp
--- a/Engine/Plugins/Runtime/USDCore/Source/USDUtilities/Private/USDGeomMeshConversion.cpp
+++ b/Engine/Plugins/Runtime/USDCore/Source/USDUtilities/Private/USDGeomMeshConversion.cpp
@@ -6,6 +6,7 @@
 
 #include "UnrealUSDWrapper.h"
 #include "USDAttributeUtils.h"
+#include "USDClassesModule.h"
 #include "USDConversionUtils.h"
 #include "USDDrawModeComponent.h"
 #include "USDLayerUtils.h"
@@ -22,6 +23,10 @@
 #include "UsdWrappers/UsdPrim.h"
 #include "UsdWrappers/UsdStage.h"
 
+#include "Animation/AnimSequence.h"
+#include "Animation/SkeletalMeshActor.h"
+#include "Components/SkeletalMeshComponent.h"
+#include "Engine/SkeletalMesh.h"
 #include "Engine/StaticMesh.h"
 #include "GeometryCache.h"
 #include "GeometryCacheMeshData.h"
@@ -30,6 +35,8 @@
 #include "Materials/MaterialInstanceDynamic.h"
 #include "MeshDescription.h"
 #include "Misc/Paths.h"
+#include "Rendering/SkeletalMeshRenderData.h"
+#include "SkeletalRenderPublic.h"
 #include "StaticMeshAttributes.h"
 #include "StaticMeshOperations.h"
 #include "StaticMeshResources.h"
@@ -4196,7 +4203,11 @@ namespace UE::UsdGeometryCacheConversion::Private
 			FrameRate = (InclusiveEndFrame - GeometryCache.GetStartFrame()) / GeometryCache.CalculateDuration();
 		}
 
-		const TArray<FName>& SlotNames;
+		FGeometryCacheExportContext(const TArray<FName>& InSlotNames)
+			: SlotNames(InSlotNames)
+		{}
+
+		TArray<FName> SlotNames;
 		int32 InclusiveEndFrame;
 		float FrameRate;
 
@@ -4207,13 +4218,266 @@ namespace UE::UsdGeometryCacheConversion::Private
 		pxr::VtArray<int> FaceVertexIndices;
 	};
 
-	void ConvertGeometryCacheMeshData(
-		const FGeometryCacheMeshData& MeshData,
-		pxr::UsdGeomMesh& UsdMesh,
+	class IUnrealMeshData
+	{
+	public:
+		virtual int32 GetNumVertices() const = 0;
+		virtual int32 GetNumIndices() const = 0;
+		virtual FVector3f GetVertexPosition(int32 Index) const = 0;
+		virtual int32 GetVertexIndex(int32 Index) const = 0;
+		virtual bool IsValid() const = 0;
+		virtual bool HasNormals() const = 0;
+		virtual FVector GetNormal(int32 Index) const = 0;
+		virtual bool HasUV() const = 0;
+		virtual FVector2D GetUV(int32 Index) const = 0;
+		virtual bool HasVertexColors() const = 0;
+		virtual FColor GetVertexColor(int32 Index) const = 0;
+		virtual bool HasVelocities() const = 0;
+		virtual FVector3f GetVelocity(int32 Index) const = 0;
+		virtual int32 GetNumSections() const = 0;
+		virtual uint32 GetSectionNumTriangles(int32 SectionIndex) const = 0;
+		virtual uint32 GetSectionStartIndex(int32 SectionIndex) const = 0;
+		virtual uint32 GetSectionMaterialIndex(int32 SectionIndex) const = 0;
+		virtual bool HasMaterialAssignment() const = 0;
+		virtual ~IUnrealMeshData()
+		{
+		}
+	};
+
+	class FGeometryCacheMeshDataWrapper : public IUnrealMeshData
+	{
+	public:
+		FGeometryCacheMeshDataWrapper(const FGeometryCacheMeshData& InMeshData)
+			: MeshData(InMeshData)
+		{
+		}
+
+		int32 GetNumVertices() const override
+		{
+			return MeshData.Positions.Num();
+		}
+		int32 GetNumIndices() const override
+		{
+			return MeshData.Indices.Num();
+		}
+		FVector3f GetVertexPosition(int32 Index) const override
+		{
+			return MeshData.Positions[Index];
+		}
+		int32 GetVertexIndex(int32 Index) const override
+		{
+			return MeshData.Indices[Index];
+		}
+		bool IsValid() const override
+		{
+			return MeshData.Positions.Num() > 0;
+		};
+		bool HasNormals() const override
+		{
+			return MeshData.VertexInfo.bHasTangentZ;
+		}
+		FVector GetNormal(int32 Index) const override
+		{
+			return MeshData.TangentsZ[Index].ToFVector();
+		}
+		bool HasUV() const override
+		{
+			return MeshData.VertexInfo.bHasUV0;
+		}
+		FVector2D GetUV(int32 Index) const override
+		{
+			return FVector2D(MeshData.TextureCoordinates[Index]);
+		}
+		bool HasVertexColors() const override
+		{
+			return MeshData.VertexInfo.bHasColor0;
+		}
+		FColor GetVertexColor(int32 Index) const override
+		{
+			return MeshData.Colors[Index];
+		}
+		bool HasVelocities() const override
+		{
+			return MeshData.VertexInfo.bHasMotionVectors;
+		}
+		FVector3f GetVelocity(int32 Index) const override
+		{
+			return MeshData.MotionVectors[Index];
+		}
+		int32 GetNumSections() const override
+		{
+			return MeshData.BatchesInfo.Num();
+		}
+		uint32 GetSectionNumTriangles(int32 SectionIndex) const override
+		{
+			return MeshData.BatchesInfo[SectionIndex].NumTriangles;
+		}
+		uint32 GetSectionStartIndex(int32 SectionIndex) const override
+		{
+			return MeshData.BatchesInfo[SectionIndex].StartIndex;
+		}
+		uint32 GetSectionMaterialIndex(int32 SectionIndex) const override
+		{
+			return MeshData.BatchesInfo[SectionIndex].MaterialIndex;
+		}
+		bool HasMaterialAssignment() const override
+		{
+			return true;
+		}
+
+	private:
+		const FGeometryCacheMeshData& MeshData;
+	};
+
+	class FSkeletalMeshDataWrapper : public IUnrealMeshData
+	{
+	public:
+		FSkeletalMeshDataWrapper(const FSkeletalMeshLODRenderData& InLODData, bool bVertexColors)
+			: LODData(InLODData)
+			, bHasVertexColors(bVertexColors)
+		{
+		}
+
+		int32 GetNumVertices() const override
+		{
+			return LODData.GetNumVertices();
+		}
+		int32 GetNumIndices() const override
+		{
+			return LODData.MultiSizeIndexContainer.GetIndexBuffer()->Num();
+		}
+		FVector3f GetVertexPosition(int32 Index) const override
+		{
+			return LODData.StaticVertexBuffers.PositionVertexBuffer.VertexPosition(Index);
+		}
+		int32 GetVertexIndex(int32 Index) const override
+		{
+			return LODData.MultiSizeIndexContainer.GetIndexBuffer()->Get(Index);
+		}
+		bool IsValid() const override
+		{
+			return GetNumVertices() > 0;
+		};
+		bool HasNormals() const override
+		{
+			return true;
+		}
+		FVector GetNormal(int32 Index) const override
+		{
+			return FVector(LODData.StaticVertexBuffers.StaticMeshVertexBuffer.VertexTangentZ(Index));
+		}
+		bool HasUV() const override
+		{
+			return LODData.StaticVertexBuffers.StaticMeshVertexBuffer.GetNumTexCoords() > 0;
+		}
+		FVector2D GetUV(int32 Index) const override
+		{
+			const int32 UVIndex = 0;
+			return FVector2D(LODData.StaticVertexBuffers.StaticMeshVertexBuffer.GetVertexUV(Index, UVIndex));
+		}
+		bool HasVertexColors() const override
+		{
+			return bHasVertexColors;
+		}
+		FColor GetVertexColor(int32 Index) const override
+		{
+			return LODData.StaticVertexBuffers.ColorVertexBuffer.VertexColor(Index);
+		}
+		bool HasVelocities() const override
+		{
+			return false;
+		}
+		FVector3f GetVelocity(int32 Index) const override
+		{
+			return {};
+		}
+		int32 GetNumSections() const override
+		{
+			return LODData.RenderSections.Num();
+		}
+		uint32 GetSectionNumTriangles(int32 SectionIndex) const override
+		{
+			return LODData.RenderSections[SectionIndex].NumTriangles;
+		}
+		uint32 GetSectionStartIndex(int32 SectionIndex) const override
+		{
+			return LODData.RenderSections[SectionIndex].BaseIndex;
+		}
+		uint32 GetSectionMaterialIndex(int32 SectionIndex) const override
+		{
+			return LODData.RenderSections[SectionIndex].MaterialIndex;
+		}
+		bool HasMaterialAssignment() const override
+		{
+			return true;
+		}
+
+	private:
+		const FSkeletalMeshLODRenderData& LODData;
+		bool bHasVertexColors;
+	};
+
+	// Wrapper for skinned vertices of a skeletal mesh as rendered through a skeletal mesh component
+	class FSkinnedVerticesDataWrapper : public FSkeletalMeshDataWrapper
+	{
+	public:
+		FSkinnedVerticesDataWrapper(const TArray<FFinalSkinVertex>& InSkinnedVertices, const FSkeletalMeshLODRenderData& InLODData, bool bVertexColors)
+			: FSkeletalMeshDataWrapper(InLODData, bVertexColors)
+			, SkinnedVertices(InSkinnedVertices)
+		{
+		}
+
+		int32 GetNumVertices() const override
+		{
+			return SkinnedVertices.Num();
+		}
+		FVector3f GetVertexPosition(int32 Index) const override
+		{
+			return SkinnedVertices[Index].Position;
+		}
+		bool HasNormals() const override
+		{
+			return true;
+		}
+		FVector GetNormal(int32 Index) const override
+		{
+			return SkinnedVertices[Index].TangentZ.ToFVector();
+		}
+		bool HasUV() const override
+		{
+			// Use the UV of the skeletal mesh
+			return false;
+		}
+		// In case we want to use the UV from the skinned vertices, but they are not animated anyway
+		//FVector2D GetUV(int32 Index) const override
+		//{
+		//	return SkinnedVertices[Index].TextureCoordinates[0];
+		//}
+		bool HasVertexColors() const override
+		{
+			return false;
+		}
+		FColor GetVertexColor(int32 Index) const override
+		{
+			return {};
+		}
+		bool HasMaterialAssignment() const override
+		{
+			// Use the material assignment/section info of the skeletal mesh
+			return false;
+		}
+
+	private:
+		const TArray<FFinalSkinVertex>& SkinnedVertices;
+	};
+
+	void ConvertMeshData(
+		const IUnrealMeshData& MeshData,
 		const TArray<FString>& MaterialAssignments,
 		const pxr::UsdTimeCode TimeCode,
 		pxr::UsdPrim PrimToReceiveMaterialAssignments,
-		FGeometryCacheExportContext& ExportContext
+		FGeometryCacheExportContext& ExportContext,
+		pxr::UsdGeomMesh& UsdMesh
 	)
 	{
 		pxr::UsdPrim MeshPrim = UsdMesh.GetPrim();
@@ -4226,7 +4490,7 @@ namespace UE::UsdGeometryCacheConversion::Private
 
 		// Vertices
 		{
-			const int32 VertexCount = MeshData.Positions.Num();
+			const int32 VertexCount = MeshData.GetNumVertices();
 
 			// Points
 			{
@@ -4238,7 +4502,7 @@ namespace UE::UsdGeometryCacheConversion::Private
 
 					for (int32 VertexIndex = 0; VertexIndex < VertexCount; ++VertexIndex)
 					{
-						PointsArray.push_back(UnrealToUsd::ConvertVectorFloat(StageInfo, (FVector)MeshData.Positions[VertexIndex]));
+						PointsArray.push_back(UnrealToUsd::ConvertVectorFloat(StageInfo, (FVector)MeshData.GetVertexPosition(VertexIndex)));
 					}
 
 					Points.Set(PointsArray, TimeCode);
@@ -4246,7 +4510,7 @@ namespace UE::UsdGeometryCacheConversion::Private
 			}
 
 			// Normals
-			if (MeshData.VertexInfo.bHasTangentZ)
+			if (MeshData.HasNormals())
 			{
 				// We need to emit this if we're writing normals (which we always are) because any DCC that can
 				// actually subdivide (like usdview) will just discard authored normals and fully recompute them
@@ -4265,7 +4529,7 @@ namespace UE::UsdGeometryCacheConversion::Private
 
 					for (int32 VertexIndex = 0; VertexIndex < VertexCount; ++VertexIndex)
 					{
-						FVector VertexNormal = MeshData.TangentsZ[VertexIndex].ToFVector();
+						FVector VertexNormal = MeshData.GetNormal(VertexIndex);
 						Normals.push_back(UnrealToUsd::ConvertVectorFloat(StageInfo, VertexNormal));
 					}
 
@@ -4274,14 +4538,14 @@ namespace UE::UsdGeometryCacheConversion::Private
 			}
 
 			// UVs
-			if (MeshData.VertexInfo.bHasUV0)
+			if (MeshData.HasUV())
 			{
 				// Only one UV set is supported
 				const int32 TexCoordSourceIndex = 0;
 				pxr::TfToken UsdUVSetName = UsdUtils::GetUVSetName(TexCoordSourceIndex).Get();
 
 				pxr::UsdGeomPrimvar PrimvarST = pxr::UsdGeomPrimvarsAPI(MeshPrim)
-													.CreatePrimvar(UsdUVSetName, pxr::SdfValueTypeNames->TexCoord2fArray, pxr::UsdGeomTokens->vertex);
+					.CreatePrimvar(UsdUVSetName, pxr::SdfValueTypeNames->TexCoord2fArray, pxr::UsdGeomTokens->vertex);
 
 				if (PrimvarST)
 				{
@@ -4289,7 +4553,7 @@ namespace UE::UsdGeometryCacheConversion::Private
 
 					for (int32 VertexIndex = 0; VertexIndex < VertexCount; ++VertexIndex)
 					{
-						FVector2D TexCoord = FVector2D(MeshData.TextureCoordinates[VertexIndex]);
+						FVector2D TexCoord = FVector2D(MeshData.GetUV(VertexIndex));
 						TexCoord[1] = 1.f - TexCoord[1];
 
 						UVs.push_back(UnrealToUsd::ConvertVectorFloat(TexCoord));
@@ -4300,7 +4564,7 @@ namespace UE::UsdGeometryCacheConversion::Private
 			}
 
 			// Vertex colors
-			if (MeshData.VertexInfo.bHasColor0)
+			if (MeshData.HasVertexColors())
 			{
 				pxr::UsdGeomPrimvar DisplayColorPrimvar = UsdMesh.CreateDisplayColorPrimvar(pxr::UsdGeomTokens->vertex);
 				pxr::UsdGeomPrimvar DisplayOpacityPrimvar = UsdMesh.CreateDisplayOpacityPrimvar(pxr::UsdGeomTokens->vertex);
@@ -4315,7 +4579,7 @@ namespace UE::UsdGeometryCacheConversion::Private
 
 					for (int32 VertexIndex = 0; VertexIndex < VertexCount; ++VertexIndex)
 					{
-						const FColor& VertexColor = MeshData.Colors[VertexIndex];
+						const FColor& VertexColor = MeshData.GetVertexColor(VertexIndex);
 
 						// The color in the MeshData is already stored as linear
 						pxr::GfVec4f Color = UnrealToUsd::ConvertColor(VertexColor.ReinterpretAsLinear());
@@ -4329,7 +4593,7 @@ namespace UE::UsdGeometryCacheConversion::Private
 			}
 
 			// Velocities
-			if (MeshData.VertexInfo.bHasMotionVectors)
+			if (MeshData.HasVelocities())
 			{
 				pxr::UsdAttribute VelocitiesAttribute = UsdMesh.CreateVelocitiesAttr();
 				if (VelocitiesAttribute)
@@ -4341,7 +4605,7 @@ namespace UE::UsdGeometryCacheConversion::Private
 					{
 						// The motion vectors in the MeshData are stored as unit per frame so convert it back to unit per second
 						Velocities.push_back(
-							UnrealToUsd::ConvertVectorFloat(StageInfo, (FVector)-MeshData.MotionVectors[VertexIndex] * ExportContext.FrameRate)
+							UnrealToUsd::ConvertVectorFloat(StageInfo, (FVector)-MeshData.GetVelocity(VertexIndex) * ExportContext.FrameRate)
 						);
 					}
 
@@ -4352,7 +4616,7 @@ namespace UE::UsdGeometryCacheConversion::Private
 
 		// Faces
 		{
-			const int32 NumIndices = MeshData.Indices.Num();
+			const int32 NumIndices = MeshData.GetNumIndices();
 			const int32 FaceCount = NumIndices / 3;
 			// Face Vertex Counts
 			{
@@ -4387,7 +4651,7 @@ namespace UE::UsdGeometryCacheConversion::Private
 
 					for (int32 Index = 0; Index < NumIndices; ++Index)
 					{
-						FaceVertexIndices.push_back(MeshData.Indices[Index]);
+						FaceVertexIndices.push_back(MeshData.GetVertexIndex(Index));
 					}
 
 					if (ExportContext.FaceVertexIndices != FaceVertexIndices)
@@ -4400,6 +4664,7 @@ namespace UE::UsdGeometryCacheConversion::Private
 		}
 
 		// Material assignments
+		if (MeshData.HasMaterialAssignment())
 		{
 			// This LOD has a single material assignment, just create/bind an UnrealMaterial child prim directly
 			if (MaterialAssignments.Num() == 1)
@@ -4411,10 +4676,8 @@ namespace UE::UsdGeometryCacheConversion::Private
 			{
 				TSet<FString> UsedSectionNames;
 				// Need to fetch all triangles of a section, and add their indices
-				for (int32 SectionIndex = 0; SectionIndex < MeshData.BatchesInfo.Num(); ++SectionIndex)
+				for (int32 SectionIndex = 0; SectionIndex < MeshData.GetNumSections(); ++SectionIndex)
 				{
-					const FGeometryCacheMeshBatchInfo& Section = MeshData.BatchesInfo[SectionIndex];
-
 					// Note that we will continue authoring the GeomSubsets on even if we later find out we have no material assignment (just
 					// "") for this section, so as to satisfy the "partition" family condition (below)
 					FString SectionName;
@@ -4444,7 +4707,7 @@ namespace UE::UsdGeometryCacheConversion::Private
 						);
 					}
 
-					pxr::UsdGeomSubset GeomSubsetSchema{GeomSubsetPrim};
+					pxr::UsdGeomSubset GeomSubsetSchema{ GeomSubsetPrim };
 
 					// Element type attribute
 					// Write the geomsubset attributes only once since they are at Default time anyway
@@ -4454,8 +4717,8 @@ namespace UE::UsdGeometryCacheConversion::Private
 						ElementTypeAttr.Set(pxr::UsdGeomTokens->face);
 
 						// Indices attribute
-						const uint32 TriangleCount = Section.NumTriangles;
-						const uint32 FirstTriangleIndex = Section.StartIndex / 3;	 // StartIndex is the first *vertex* instance index
+						const uint32 TriangleCount = MeshData.GetSectionNumTriangles(SectionIndex);
+						const uint32 FirstTriangleIndex = MeshData.GetSectionStartIndex(SectionIndex) / 3;	 // StartIndex is the first *vertex* instance index
 						pxr::VtArray<int> IndicesAttrValue;
 						for (uint32 TriangleIndex = FirstTriangleIndex; TriangleIndex - FirstTriangleIndex < TriangleCount; ++TriangleIndex)
 						{
@@ -4478,9 +4741,10 @@ namespace UE::UsdGeometryCacheConversion::Private
 						pxr::UsdGeomSubset::SetFamilyType(UsdMesh, pxr::UsdShadeTokens->materialBind, pxr::UsdGeomTokens->partition);
 
 						// material:binding relationship
-						if (MaterialAssignments.IsValidIndex(Section.MaterialIndex))
+						uint32 MaterialIndex = MeshData.GetSectionMaterialIndex(SectionIndex);
+						if (MaterialAssignments.IsValidIndex(MaterialIndex))
 						{
-							UsdUtils::AuthorUnrealMaterialBinding(MaterialGeomSubsetPrim, MaterialAssignments[Section.MaterialIndex]);
+							UsdUtils::AuthorUnrealMaterialBinding(MaterialGeomSubsetPrim, MaterialAssignments[MaterialIndex]);
 						}
 					}
 				}
@@ -4491,7 +4755,7 @@ namespace UE::UsdGeometryCacheConversion::Private
 
 bool UnrealToUsd::ConvertGeometryCache(const UGeometryCache* GeometryCache, pxr::UsdPrim& UsdPrim, UE::FUsdStage* StageForMaterialAssignments)
 {
-	namespace UsdGeometryCacheImpl = UE::UsdGeometryCacheConversion::Private;
+	using namespace UE::UsdGeometryCacheConversion::Private;
 
 	FScopedUsdAllocs UsdAllocs;
 
@@ -4503,8 +4767,6 @@ bool UnrealToUsd::ConvertGeometryCache(const UGeometryCache* GeometryCache, pxr:
 
 	const FUsdStageInfo StageInfo(Stage);
 
-	pxr::SdfPath ParentPrimPath = UsdPrim.GetPath();
-
 	// Collect all material assignments, referenced by the sections' material indices
 	bool bHasMaterialAssignments = false;
 	TArray<FString> MaterialAssignments;
@@ -4543,14 +4805,14 @@ bool UnrealToUsd::ConvertGeometryCache(const UGeometryCache* GeometryCache, pxr:
 	pxr::UsdPrim MaterialPrim = MaterialStage->OverridePrim(UsdPrim.GetPath());
 	pxr::UsdAttribute ExtentsAttr = TargetMesh ? TargetMesh.CreateExtentAttr() : pxr::UsdAttribute{};
 
-	UsdGeometryCacheImpl::FGeometryCacheExportContext ExportContext(*GeometryCache);
+	FGeometryCacheExportContext ExportContext(*GeometryCache);
 	const int32 StartFrame = GeometryCache->GetStartFrame();
 	const int32 EndFrame = ExportContext.InclusiveEndFrame;
 	int32 ActualStartFrame = -1;
 
 	for (int32 FrameIndex = StartFrame; FrameIndex <= EndFrame; ++FrameIndex)
 	{
-		FGeometryCacheMeshData MeshData = UsdGeometryCacheImpl::GetFlattenedGeometryCacheMeshData(GeometryCache, FrameIndex - StartFrame);
+		FGeometryCacheMeshData MeshData = GetFlattenedGeometryCacheMeshData(GeometryCache, FrameIndex - StartFrame);
 		// First frame of the animation cannot be empty otherwise the geometry cache translator would not be able to detect the animation
 		// It is allowed to have empty frames during or at the end of the animation, eg. for fluid sim or FX that disappear
 		const bool bIsValidFrame = MeshData.Positions.Num() > 0 || ActualStartFrame > 0;
@@ -4562,7 +4824,7 @@ bool UnrealToUsd::ConvertGeometryCache(const UGeometryCache* GeometryCache, pxr:
 				ActualStartFrame = FrameIndex;
 			}
 			const float TimeCode = FrameIndex;
-			UsdGeometryCacheImpl::ConvertGeometryCacheMeshData(MeshData, TargetMesh, MaterialAssignments, TimeCode, MaterialPrim, ExportContext);
+			ConvertMeshData(FGeometryCacheMeshDataWrapper(MeshData), MaterialAssignments, TimeCode, MaterialPrim, ExportContext, TargetMesh);
 		}
 
 		if (MeshData.BoundingBox.IsValid && ExtentsAttr)
@@ -4580,6 +4842,283 @@ bool UnrealToUsd::ConvertGeometryCache(const UGeometryCache* GeometryCache, pxr:
 	return true;
 }
 
+bool UnrealToUsd::ConvertSkeletalMeshToStaticMesh(
+	const USkeletalMesh* SkeletalMesh,
+	pxr::UsdPrim& UsdPrim,
+	const pxr::UsdTimeCode TimeCode,
+	UE::FUsdStage* StageForMaterialAssignments
+)
+{
+	using namespace UE::UsdGeometryCacheConversion::Private;
+
+	if (!SkeletalMesh)
+	{
+		return false;
+	}
+
+	FScopedUsdAllocs UsdAllocs;
+
+	pxr::UsdGeomMesh MeshPrim(UsdPrim);
+	if (!MeshPrim)
+	{
+		return false;
+	}
+
+	pxr::UsdStageRefPtr Stage = UsdPrim.GetStage();
+	if (!Stage)
+	{
+		return false;
+	}
+
+	const FUsdStageInfo StageInfo(Stage);
+
+	// Collect all material assignments, referenced by the sections' material indices
+	bool bHasMaterialAssignments = false;
+	TArray<FString> MaterialAssignments;
+	TArray<FName> SlotNames;
+	for (const FSkeletalMaterial& Material : SkeletalMesh->GetMaterials())
+	{
+		FString AssignedMaterialPathName;
+		if (Material.MaterialInterface->GetOutermost() != GetTransientPackage())
+		{
+			AssignedMaterialPathName = Material.MaterialInterface->GetPathName();
+			SlotNames.Add(Material.MaterialSlotName);
+			bHasMaterialAssignments = true;
+		}
+
+		MaterialAssignments.Add(AssignedMaterialPathName);
+	}
+	if (!bHasMaterialAssignments)
+	{
+		// Prevent creation of the unrealMaterials attribute in case we don't have any assignments at all
+		MaterialAssignments.Reset();
+	}
+
+	// Author material bindings on the dedicated stage if we have one
+	pxr::UsdStageRefPtr MaterialStage;
+	if (StageForMaterialAssignments)
+	{
+		MaterialStage = static_cast<pxr::UsdStageRefPtr>(*StageForMaterialAssignments);
+	}
+	else
+	{
+		MaterialStage = Stage;
+	}
+
+	const FSkeletalMeshRenderData& SkelMeshRenderData = *SkeletalMesh->GetResourceForRendering();
+
+	const int32 LODLevel = 0;
+	const FSkeletalMeshLODRenderData& LODData = SkelMeshRenderData.LODRenderData[LODLevel];
+
+	FGeometryCacheExportContext ExportContext(SlotNames);
+	pxr::UsdPrim MaterialPrim = MaterialStage->OverridePrim(UsdPrim.GetPath());
+	ConvertMeshData(FSkeletalMeshDataWrapper(LODData, SkeletalMesh->GetHasVertexColors()), MaterialAssignments, TimeCode, MaterialPrim, ExportContext, MeshPrim);
+
+	pxr::UsdAttribute ExtentsAttr = MeshPrim.CreateExtentAttr();
+	const FBox BoundingBox = SkeletalMesh->GetBounds().GetBox();
+	if (BoundingBox.IsValid && ExtentsAttr)
+	{
+		TUsdStore<pxr::VtArray<pxr::GfVec3f>> USDBounds = UnrealToUsd::ConvertBounds(StageInfo, BoundingBox);
+		ExtentsAttr.Set(USDBounds.Get(), TimeCode);
+	}
+
+	return true;
+}
+
+bool UnrealToUsd::ConvertAnimSequenceToAnimatedMesh(
+	UAnimSequence* AnimSequence,
+	USkeletalMesh* SkeletalMesh,
+	pxr::UsdPrim& UsdPrim,
+	UE::FUsdStage* StageForMaterialAssignments
+)
+{
+	using namespace UE::UsdGeometryCacheConversion::Private;
+
+	UWorld* World = IUsdClassesModule::GetCurrentWorld();
+	if (!AnimSequence || !SkeletalMesh || !World)
+	{
+		return false;
+	}
+
+	FScopedUsdAllocs UsdAllocs;
+
+	pxr::UsdGeomMesh MeshPrim(UsdPrim);
+	if (!MeshPrim)
+	{
+		return false;
+	}
+
+	pxr::UsdStageRefPtr Stage = UsdPrim.GetStage();
+	if (!Stage)
+	{
+		return false;
+	}
+
+	// Collect all material assignments, referenced by the sections' material indices
+	bool bHasMaterialAssignments = false;
+	TArray<FString> MaterialAssignments;
+	TArray<FName> SlotNames;
+	for (const FSkeletalMaterial& Material : SkeletalMesh->GetMaterials())
+	{
+		FString AssignedMaterialPathName;
+		if (Material.MaterialInterface->GetOutermost() != GetTransientPackage())
+		{
+			AssignedMaterialPathName = Material.MaterialInterface->GetPathName();
+			SlotNames.Add(Material.MaterialSlotName);
+			bHasMaterialAssignments = true;
+		}
+
+		MaterialAssignments.Add(AssignedMaterialPathName);
+	}
+	if (!bHasMaterialAssignments)
+	{
+		// Prevent creation of the unrealMaterials attribute in case we don't have any assignments at all
+		MaterialAssignments.Reset();
+	}
+
+	// Author material bindings on the dedicated stage if we have one
+	pxr::UsdStageRefPtr MaterialStage;
+	if (StageForMaterialAssignments)
+	{
+		MaterialStage = static_cast<pxr::UsdStageRefPtr>(*StageForMaterialAssignments);
+	}
+	else
+	{
+		MaterialStage = Stage;
+	}
+
+	pxr::UsdPrim MaterialPrim = MaterialStage->OverridePrim(UsdPrim.GetPath());
+
+	// Create a temp SkeletalMeshActor on which to play back the AnimSequence and get the skinned vertices to bake
+	FActorSpawnParameters SpawnParameters;
+	SpawnParameters.ObjectFlags = RF_Transient;
+	SpawnParameters.Name = *UsdToUnreal::ConvertString(UsdPrim.GetName());
+	SpawnParameters.NameMode = FActorSpawnParameters::ESpawnActorNameMode::Requested;	 // Will generate a unique name in case of a conflict
+
+	ASkeletalMeshActor* SkelMeshActor = Cast<ASkeletalMeshActor>(World->SpawnActor(ASkeletalMeshActor::StaticClass(), nullptr, SpawnParameters));
+
+	USkeletalMeshComponent* SkelMeshComponent = SkelMeshActor->GetSkeletalMeshComponent();
+	SkelMeshComponent->SetSkeletalMeshAsset(SkeletalMesh);
+	SkelMeshComponent->OverrideAnimationData(AnimSequence);
+
+	const FSkeletalMeshRenderData& SkelMeshRenderData = *SkeletalMesh->GetResourceForRendering();
+
+	const int32 LODLevel = 0;
+	const FSkeletalMeshLODRenderData& LODData = SkelMeshRenderData.LODRenderData[LODLevel];
+
+	const int32 StartFrame = 0;
+	const int32 EndFrame = AnimSequence->GetNumberOfSampledKeys() - 1;
+	const double FPS = AnimSequence->GetSamplingFrameRate().AsDecimal();
+	const bool bHasVertexColors = SkeletalMesh->GetHasVertexColors();
+	FGeometryCacheExportContext ExportContext(SlotNames);
+	for (int32 FrameIndex = StartFrame; FrameIndex <= EndFrame; ++FrameIndex)
+	{
+		double Time = FrameIndex / FPS;
+		SkelMeshComponent->SetPosition((float)Time);
+
+		if (USkeletalMeshComponent* Leader = Cast<USkeletalMeshComponent>(SkelMeshComponent->LeaderPoseComponent.Get()))
+		{
+			UsdUtils::RefreshSkeletalMeshComponent(*Leader);
+		}
+		UsdUtils::RefreshSkeletalMeshComponent(*SkelMeshComponent);
+
+		if (AActor* Owner = SkelMeshComponent->GetOwner())
+		{
+			Owner->Tick(0.0f);
+		}
+
+		TArray<FFinalSkinVertex> SkinnedVertices;
+		SkelMeshComponent->GetCPUSkinnedVertices(SkinnedVertices, LODLevel);
+
+		ConvertMeshData(FSkinnedVerticesDataWrapper(SkinnedVertices, LODData, bHasVertexColors), MaterialAssignments, FrameIndex, MaterialPrim, ExportContext, MeshPrim);
+	}
+
+	World->DestroyActor(SkelMeshActor);
+
+	return true;
+}
+
+bool UnrealToUsd::CreateSkeletalAnimationToMeshBaker(UE::FUsdPrim& UsdPrim, USkeletalMeshComponent& SkelMeshComponent, UnrealToUsd::FComponentBaker& OutBaker)
+{
+	using namespace UE::UsdGeometryCacheConversion::Private;
+
+	USkeletalMesh* SkeletalMesh = SkelMeshComponent.GetSkeletalMeshAsset();
+	if (!SkeletalMesh)
+	{
+		return false;
+	}
+
+	FScopedUsdAllocs Allocs;
+
+	pxr::UsdGeomMesh MeshPrim(UsdPrim);
+	if (!MeshPrim)
+	{
+		return false;
+	}
+
+	pxr::UsdStageRefPtr Stage = UsdPrim.GetStage();
+	if (!Stage)
+	{
+		return false;
+	}
+
+	// Collect all material assignments, referenced by the sections' material indices
+	bool bHasMaterialAssignments = false;
+	TArray<FString> MaterialAssignments;
+	TArray<FName> SlotNames;
+	for (const FSkeletalMaterial& Material : SkeletalMesh->GetMaterials())
+	{
+		FString AssignedMaterialPathName;
+		if (Material.MaterialInterface->GetOutermost() != GetTransientPackage())
+		{
+			AssignedMaterialPathName = Material.MaterialInterface->GetPathName();
+			SlotNames.Add(Material.MaterialSlotName);
+			bHasMaterialAssignments = true;
+		}
+
+		MaterialAssignments.Add(AssignedMaterialPathName);
+	}
+	if (!bHasMaterialAssignments)
+	{
+		// Prevent creation of the unrealMaterials attribute in case we don't have any assignments at all
+		MaterialAssignments.Reset();
+	}
+
+	OutBaker.ComponentPath = SkelMeshComponent.GetPathName();
+	OutBaker.BakerType = UnrealToUsd::EBakingType::Skeletal;
+
+	const FSkeletalMeshRenderData& SkelMeshRenderData = *SkeletalMesh->GetResourceForRendering();
+
+	const int32 LODLevel = 0;
+	const FSkeletalMeshLODRenderData& LODData = SkelMeshRenderData.LODRenderData[LODLevel];
+
+	const bool bHasVertexColors = SkeletalMesh->GetHasVertexColors();
+	OutBaker.BakerFunction = [&SkelMeshComponent, &LODData, MaterialAssignments, UsdPrim, bHasVertexColors, ExportContext = FGeometryCacheExportContext(SlotNames)](double UsdTimeCode) mutable
+	{
+		FScopedUsdAllocs InnerAllocs;
+
+		if (USkeletalMeshComponent* Leader = Cast<USkeletalMeshComponent>(SkelMeshComponent.LeaderPoseComponent.Get()))
+		{
+			UsdUtils::RefreshSkeletalMeshComponent(*Leader);
+		}
+		UsdUtils::RefreshSkeletalMeshComponent(SkelMeshComponent);
+
+		if (AActor* Owner = SkelMeshComponent.GetOwner())
+		{
+			Owner->Tick(0.0f);
+		}
+
+		const int32 LODLevel = 0;
+		TArray<FFinalSkinVertex> SkinnedVertices;
+		SkelMeshComponent.GetCPUSkinnedVertices(SkinnedVertices, LODLevel);
+
+		pxr::UsdGeomMesh MeshPrim(UsdPrim);
+		ConvertMeshData(FSkinnedVerticesDataWrapper(SkinnedVertices, LODData, bHasVertexColors), MaterialAssignments, UsdTimeCode, UsdPrim, ExportContext, MeshPrim);
+	};
+
+	return true;
+}
+
 void UsdUtils::RepairNormalsAndTangents(const FString& PrimPath, FMeshDescription& MeshDescription)
 {
 	FStaticMeshConstAttributes Attributes{MeshDescription};
diff --git a/Engine/Plugins/Runtime/USDCore/Source/USDUtilities/Private/USDPrimConversion.cpp b/Engine/Plugins/Runtime/USDCore/Source/USDUtilities/Private/USDPrimConversion.cpp
--- a/Engine/Plugins/Runtime/USDCore/Source/USDUtilities/Private/USDPrimConversion.cpp
+++ b/Engine/Plugins/Runtime/USDCore/Source/USDUtilities/Private/USDPrimConversion.cpp
@@ -4377,73 +4377,6 @@ bool UnrealToUsd::CreateComponentPropertyBaker(
 	return false;
 }
 
-namespace UE::USDPrimConversion::Private
-{
-	// Returns bone-space joint transforms from the SkeletalMeshComponent while paying attention to whether
-	// it has a LeaderPoseComponent or not.
-	//
-	// References:
-	// - FMLDeformerEditorToolkit::GetDebugActorComponentSpaceTransforms
-	// - FAnimationRecorder::GetBoneTransforms
-	void GetBoneTransforms(USkeletalMeshComponent* Component, TArray<FTransform>& BoneTransforms)
-	{
-		if (!Component)
-		{
-			return;
-		}
-
-		int32 NumBones = INDEX_NONE;
-		if (USkeletalMesh* Mesh = Component->GetSkeletalMeshAsset())
-		{
-			const FReferenceSkeleton& RefSkel = Mesh->GetRefSkeleton();
-			NumBones = RefSkel.GetNum();
-		}
-		if (NumBones == INDEX_NONE)
-		{
-			return;
-		}
-
-		if (USkeletalMeshComponent* Leader = Cast<USkeletalMeshComponent>(Component->LeaderPoseComponent.Get()))
-		{
-			const TArray<FTransform>& LeaderTransforms = Leader->GetBoneSpaceTransforms();
-			const TArray<FTransform>& FollowerTransforms = Component->GetBoneSpaceTransforms();
-
-			const TArray<int32>& BoneMap = Component->GetLeaderBoneMap();
-
-			BoneTransforms.SetNumUninitialized(NumBones);
-			for (int32 BoneIndex = 0; BoneIndex < NumBones; BoneIndex++)
-			{
-				if (BoneMap.IsValidIndex(BoneIndex) && LeaderTransforms.IsValidIndex(BoneMap[BoneIndex]))
-				{
-					BoneTransforms[BoneIndex] = LeaderTransforms[BoneMap[BoneIndex]];
-				}
-				else if (FollowerTransforms.IsValidIndex(BoneIndex))
-				{
-					BoneTransforms[BoneIndex] = FollowerTransforms[BoneIndex];
-				}
-			}
-		}
-		else
-		{
-			BoneTransforms = Component->GetBoneSpaceTransforms();
-		}
-	}
-
-	void RefreshSkeletalMeshComponent(USkeletalMeshComponent& Component)
-	{
-		// This whole incantation is required or else the component will really not update until the next frame.
-		// Note: This will also cause the update of morph target weights.
-		Component.TickAnimation(0.f, false);
-		Component.UpdateLODStatus();
-		Component.RefreshBoneTransforms();
-		Component.RefreshFollowerComponents();
-		Component.UpdateComponentToWorld();
-		Component.FinalizeBoneTransform();
-		Component.MarkRenderTransformDirty();
-		Component.MarkRenderDynamicDataDirty();
-	}
-}	 // namespace UE::USDPrimConversion::Private
-
 bool UnrealToUsd::CreateSkeletalAnimationBaker(
 	UE::FUsdPrim& SkeletonPrim,
 	UE::FUsdPrim& SkelAnimation,
@@ -4571,9 +4504,9 @@ bool UnrealToUsd::CreateSkeletalAnimationBaker(
 
 		if (USkeletalMeshComponent* Leader = Cast<USkeletalMeshComponent>(Component.LeaderPoseComponent.Get()))
 		{
-			RefreshSkeletalMeshComponent(*Leader);
+			UsdUtils::RefreshSkeletalMeshComponent(*Leader);
 		}
-		RefreshSkeletalMeshComponent(Component);
+		UsdUtils::RefreshSkeletalMeshComponent(Component);
 
 		// I'm not entirely sure why this is needed but FFbxExporter::ExportAnimTrack and FFbxExporter::ExportLevelSequenceBaked3DTransformTrack
 		// do this so for safety maybe we should as well?
@@ -4583,7 +4516,7 @@ bool UnrealToUsd::CreateSkeletalAnimationBaker(
 		}
 
 		TArray<FTransform> LocalBoneTransforms;
-		GetBoneTransforms(&Component, LocalBoneTransforms);
+		UsdUtils::GetBoneTransforms(&Component, LocalBoneTransforms);
 
 		// For whatever reason it seems that sometimes this is not ready for us, so let's force it to be recalculated
 		if (LocalBoneTransforms.Num() == 0)
diff --git a/Engine/Plugins/Runtime/USDCore/Source/USDUtilities/Private/USDSkeletalDataConversion.cpp b/Engine/Plugins/Runtime/USDCore/Source/USDUtilities/Private/USDSkeletalDataConversion.cpp
--- a/Engine/Plugins/Runtime/USDCore/Source/USDUtilities/Private/USDSkeletalDataConversion.cpp
+++ b/Engine/Plugins/Runtime/USDCore/Source/USDUtilities/Private/USDSkeletalDataConversion.cpp
@@ -29,7 +29,9 @@
 #include "AnimEncoding.h"
 #include "Async/ParallelFor.h"
 #include "BoneWeights.h"
+#include "Components/SkeletalMeshComponent.h"
 #include "ControlRig.h"
+#include "Engine/SkeletalMesh.h"
 #include "Engine/SkinnedAssetCommon.h"
 #include "Evaluation/MovieSceneSequenceTransform.h"
 #include "IMovieScenePlayer.h"
@@ -3132,6 +3134,70 @@ void UsdUtils::ResolveWeightsForBlendShape(
 	}
 }
 
+// Returns bone-space joint transforms from the SkeletalMeshComponent while paying attention to whether
+// it has a LeaderPoseComponent or not.
+//
+// References:
+// - FMLDeformerEditorToolkit::GetDebugActorComponentSpaceTransforms
+// - FAnimationRecorder::GetBoneTransforms
+void UsdUtils::GetBoneTransforms(USkeletalMeshComponent* Component, TArray<FTransform>& BoneTransforms)
+{
+	if (!Component)
+	{
+		return;
+	}
+
+	int32 NumBones = INDEX_NONE;
+	if (USkeletalMesh* Mesh = Component->GetSkeletalMeshAsset())
+	{
+		const FReferenceSkeleton& RefSkel = Mesh->GetRefSkeleton();
+		NumBones = RefSkel.GetNum();
+	}
+	if (NumBones == INDEX_NONE)
+	{
+		return;
+	}
+
+	if (USkeletalMeshComponent* Leader = Cast<USkeletalMeshComponent>(Component->LeaderPoseComponent.Get()))
+	{
+		const TArray<FTransform>& LeaderTransforms = Leader->GetBoneSpaceTransforms();
+		const TArray<FTransform>& FollowerTransforms = Component->GetBoneSpaceTransforms();
+
+		const TArray<int32>& BoneMap = Component->GetLeaderBoneMap();
+
+		BoneTransforms.SetNumUninitialized(NumBones);
+		for (int32 BoneIndex = 0; BoneIndex < NumBones; BoneIndex++)
+		{
+			if (BoneMap.IsValidIndex(BoneIndex) && LeaderTransforms.IsValidIndex(BoneMap[BoneIndex]))
+			{
+				BoneTransforms[BoneIndex] = LeaderTransforms[BoneMap[BoneIndex]];
+			}
+			else if (FollowerTransforms.IsValidIndex(BoneIndex))
+			{
+				BoneTransforms[BoneIndex] = FollowerTransforms[BoneIndex];
+			}
+		}
+	}
+	else
+	{
+		BoneTransforms = Component->GetBoneSpaceTransforms();
+	}
+}
+
+void UsdUtils::RefreshSkeletalMeshComponent(USkeletalMeshComponent& Component)
+{
+	// This whole incantation is required or else the component will really not update until the next frame.
+	// Note: This will also cause the update of morph target weights.
+	Component.TickAnimation(0.f, false);
+	Component.UpdateLODStatus();
+	Component.RefreshBoneTransforms();
+	Component.RefreshFollowerComponents();
+	Component.UpdateComponentToWorld();
+	Component.FinalizeBoneTransform();
+	Component.MarkRenderTransformDirty();
+	Component.MarkRenderDynamicDataDirty();
+}
+
 #if USE_USD_SDK
 
 // Adapted from UsdSkel_CacheImpl::ReadScope::_FindOrCreateSkinningQuery because we need to manually create these on UsdGeomMeshes we already have
diff --git a/Engine/Plugins/Runtime/USDCore/Source/USDUtilities/Public/USDGeomMeshConversion.h b/Engine/Plugins/Runtime/USDCore/Source/USDUtilities/Public/USDGeomMeshConversion.h
--- a/Engine/Plugins/Runtime/USDCore/Source/USDUtilities/Public/USDGeomMeshConversion.h
+++ b/Engine/Plugins/Runtime/USDCore/Source/USDUtilities/Public/USDGeomMeshConversion.h
@@ -25,9 +25,12 @@ PXR_NAMESPACE_OPEN_SCOPE
 PXR_NAMESPACE_CLOSE_SCOPE
 #endif	  // USE_USD_SDK
 
+class UAnimSequence;
 class UGeometryCache;
 class UMaterialInstanceConstant;
 class UMaterialInstanceDynamic;
+class USkeletalMesh;
+class USkeletalMeshComponent;
 class UStaticMesh;
 struct FMeshDescription;
 struct FStaticMeshLODResources;
@@ -241,7 +244,7 @@ namespace UnrealToUsd
 	/**
 	 * Extracts animated mesh data from GeometryCache and places the results in UsdPrim.
 	 * @param GeometryCache - GeometryCache to convert
-	 * @param UsdPrim - Prim to receive the mesh data or LOD variant set
+	 * @param UsdPrim - Prim to receive the mesh data
 	 * @param StageForMaterialAssignments - Stage to use when authoring material assignments (we use this when we want to export the mesh to a payload
 	 * layer, but the material assignments to an asset layer)
 	 * @return Whether the conversion was successful or not.
@@ -251,6 +254,26 @@ namespace UnrealToUsd
 		pxr::UsdPrim& UsdPrim,
 		UE::FUsdStage* StageForMaterialAssignments = nullptr
 	);
+
+	USDUTILITIES_API bool ConvertSkeletalMeshToStaticMesh(
+		const USkeletalMesh* SkeletalMesh,
+		pxr::UsdPrim& UsdPrim,
+		const pxr::UsdTimeCode TimeCode = pxr::UsdTimeCode::Default(),
+		UE::FUsdStage* StageForMaterialAssignments = nullptr
+	);
+
+	USDUTILITIES_API bool ConvertAnimSequenceToAnimatedMesh(
+		UAnimSequence* AnimSequence,
+		USkeletalMesh* SkeletalMesh,
+		pxr::UsdPrim& UsdPrim,
+		UE::FUsdStage* StageForMaterialAssignments = nullptr
+	);
+
+	/**
+	 * Creates a baker responsible for baking the skinned mesh of a skeletal mesh component to a Mesh prim,
+	 * returning true if a EBakingType::Skeletal baker was created
+	 */
+	USDUTILITIES_API bool CreateSkeletalAnimationToMeshBaker(UE::FUsdPrim& MeshPrim, USkeletalMeshComponent& Component, struct FComponentBaker& OutBaker);
 }
 #endif	  // USE_USD_SDK
 
diff --git a/Engine/Plugins/Runtime/USDCore/Source/USDUtilities/Public/USDSkeletalDataConversion.h b/Engine/Plugins/Runtime/USDCore/Source/USDUtilities/Public/USDSkeletalDataConversion.h
--- a/Engine/Plugins/Runtime/USDCore/Source/USDUtilities/Public/USDSkeletalDataConversion.h
+++ b/Engine/Plugins/Runtime/USDCore/Source/USDUtilities/Public/USDSkeletalDataConversion.h
@@ -38,6 +38,7 @@ class UAnimSequence;
 class UMovieScene;
 class UMovieSceneControlRigParameterSection;
 class USkeletalMesh;
+class USkeletalMeshComponent;
 struct FUsdStageInfo;
 struct FMovieSceneInverseSequenceTransform;
 
@@ -124,6 +125,9 @@ namespace UsdUtils
 		TArray<float>& OutInbetweenWeights
 	);
 
+	USDUTILITIES_API void GetBoneTransforms(USkeletalMeshComponent* Component, TArray<FTransform>& BoneTransforms);
+	USDUTILITIES_API void RefreshSkeletalMeshComponent(USkeletalMeshComponent& Component);
+
 #if USE_USD_SDK
 	/** Allows creation of a skinning query from the underlying skinned mesh and skeleton. Adapted from the USD SDK implementation */
 	USDUTILITIES_API UE::FUsdSkelSkinningQuery CreateSkinningQuery(
