USD: Implement support for writing material overrides inside instanceables;

Fix the same source prim paths being added multiple times to UsdAssetUserData.

diff --git a/Engine/Plugins/Importers/USDImporter/Source/USDSchemas/Private/MeshTranslationImpl.cpp b/Engine/Plugins/Importers/USDImporter/Source/USDSchemas/Private/MeshTranslationImpl.cpp
--- a/Engine/Plugins/Importers/USDImporter/Source/USDSchemas/Private/MeshTranslationImpl.cpp
+++ b/Engine/Plugins/Importers/USDImporter/Source/USDSchemas/Private/MeshTranslationImpl.cpp
@@ -1004,8 +1004,15 @@ void MeshTranslationImpl::RecordSourcePrimsForMaterialSlots(
 
 		for (int32 LODSlotIndex = 0; LODSlotIndex < LODSlots.Num(); ++LODSlotIndex, ++SlotIndex)
 		{
+			TArray<FString>& PrimPaths = UserData->MaterialSlotToPrimPaths.FindOrAdd(SlotIndex).PrimPaths;
+
 			const UsdUtils::FUsdPrimMaterialSlot& Slot = LODSlots[LODSlotIndex];
-			UserData->MaterialSlotToPrimPaths.FindOrAdd(SlotIndex).PrimPaths.Append(Slot.PrimPaths.Array());
+
+			PrimPaths.Reserve(PrimPaths.Num() + Slot.PrimPaths.Num());
+			for (const FString& SlotPrimPath : Slot.PrimPaths)
+			{
+				PrimPaths.AddUnique(SlotPrimPath);
+			}
 		}
 	}
 }
diff --git a/Engine/Plugins/Runtime/USDCore/Source/USDUtilities/Private/USDAttributeUtils.cpp b/Engine/Plugins/Runtime/USDCore/Source/USDUtilities/Private/USDAttributeUtils.cpp
--- a/Engine/Plugins/Runtime/USDCore/Source/USDUtilities/Private/USDAttributeUtils.cpp
+++ b/Engine/Plugins/Runtime/USDCore/Source/USDUtilities/Private/USDAttributeUtils.cpp
@@ -364,7 +364,7 @@ bool UsdUtils::NotifyIfInstanceProxy(const pxr::UsdPrim& Prim)
 					UE_LOG(
 						LogUsd,
 						Log,
-						TEXT("Ignoring edits to prim '%s' as it is an instance proxy"),
+						TEXT("Ignoring some edits to prim '%s' as it is an instance proxy"),
 						*UsdToUnreal::ConvertPath(Prim.GetPrimPath())
 					);
 					return true;
diff --git a/Engine/Plugins/Runtime/USDCore/Source/USDUtilities/Private/USDPrimConversion.cpp b/Engine/Plugins/Runtime/USDCore/Source/USDUtilities/Private/USDPrimConversion.cpp
--- a/Engine/Plugins/Runtime/USDCore/Source/USDUtilities/Private/USDPrimConversion.cpp
+++ b/Engine/Plugins/Runtime/USDCore/Source/USDUtilities/Private/USDPrimConversion.cpp
@@ -3391,10 +3391,120 @@ bool UnrealToUsd::ConvertMaterialOverrides(
 
 	FScopedUsdAllocs Allocs;
 	pxr::UsdStageRefPtr Stage = UsdPrim.GetStage();
-	FString UsdPrimPath = UsdToUnreal::ConvertPath(UsdPrim.GetPrimPath());
+	FString MeshPrimPath = UsdToUnreal::ConvertPath(UsdPrim.GetPrimPath());
+
+	// If we're inside of an instance then we have to account for the fact that with instance-aware translation only one
+	// of the prototype's instances will create/fetch the static mesh and record their material slot paths. This means that
+	// in order to tell whether a slot prim path from asset user data corresponds to a child of our target mesh prim we
+	// need to compare paths within the prototype instead
+	const bool bIsInsideInstance = UsdPrim.IsInstanceProxy();
+	if (bIsInsideInstance)
+	{
+		MeshPrimPath = UsdToUnreal::ConvertPath(UsdPrim.GetPrimInPrototype().GetPrimPath());
+	}
 
 	const bool bAllLODs = LowestLOD == INDEX_NONE && HighestLOD == INDEX_NONE;
 
+	TFunction<void(const UUsdMeshAssetUserData*, int32, const TSet<FString>&, const FString&)> AuthorOverrideViaUserData =
+		[bIsInsideInstance,
+		 &Stage,
+		 &UsdPrim](
+			const UUsdMeshAssetUserData* UserData,	  //
+			int32 MatIndex,
+			const TSet<FString>& MeshPrimPaths,
+			const FString& OverrideMaterialPath
+		)
+	{
+		const FUsdPrimPathList* SourcePrimPaths = UserData->MaterialSlotToPrimPaths.Find(MatIndex);
+		if (!SourcePrimPaths)
+		{
+			return;
+		}
+
+		for (FString SourcePrimPath : SourcePrimPaths->PrimPaths)
+		{
+			if (bIsInsideInstance)
+			{
+				pxr::UsdPrim PrimAtPath = Stage->GetPrimAtPath(UnrealToUsd::ConvertPath(*SourcePrimPath).Get());
+				if (PrimAtPath.IsInstanceProxy())
+				{
+					SourcePrimPath = UsdToUnreal::ConvertPath(PrimAtPath.GetPrimInPrototype().GetPrimPath());
+				}
+			}
+
+			// Our mesh assets are shared between multiple prims via the asset cache, and all user prims of that asset will record
+			// their source prim paths for each material slot. In here we just want to apply overrides to the prims that
+			// correspond to the modified component (i.e. are within its subtree), and not the others.
+			for (const FString& MeshPrimPath : MeshPrimPaths)
+			{
+				if (SourcePrimPath.StartsWith(MeshPrimPath))
+				{
+					if (bIsInsideInstance)
+					{
+						// Have to author collection-based bindings on the instance root.
+						//
+						// SourcePrimPath and MeshPrimPath are both relative to the prototype in this case. If we know they match, we
+						// know this source prim path pertains to something inside our instanceable, and so we should author the override
+
+						// Find the instance root
+						pxr::UsdPrim InstanceRoot = UsdPrim.GetParent();
+						while (InstanceRoot && !InstanceRoot.IsPseudoRoot() && !InstanceRoot.IsInstance())
+						{
+							InstanceRoot = InstanceRoot.GetParent();
+						}
+						if (InstanceRoot.IsInstance())
+						{
+							// Now we just have a bunch of paths relative to the prototype. In order to find the correct instance proxy
+							// path to override, we must do some path surgery to move the SourcePrimPath suffix below the instance root
+							// prefix, onto the true instance proxy path on the stage.
+							//
+							// Note that we have to do this (and not just use the original SourcePrimPaths directly) mainly because with
+							// instance-aware translation it is now possible that SourcePrimPath doesn't contain the paths for *all*
+							// instances of the prototype, and only for the one instance that actually led to it being translated.
+							// UsdPrim/InstanceRoot may refer to another instance entirely, so we need our paths to match them instead.
+
+							// e.g. "/root/instanceable7"
+							pxr::SdfPath InstanceRootPath = InstanceRoot.GetPrimPath();
+
+							// e.g. "/__Prototype_3/child_prim/slot
+							pxr::SdfPath InstanceProxySuffix = UnrealToUsd::ConvertPath(*SourcePrimPath).Get();
+
+							// After the loop PrototypePrefix becomes e.g. "/__Prototype_3"
+							pxr::SdfPath PrototypePrefix = InstanceProxySuffix;
+							while (true)
+							{
+								pxr::SdfPath ParentPath = PrototypePrefix.GetParentPath();
+								if (ParentPath.IsEmpty() || ParentPath.IsAbsoluteRootPath())
+								{
+									break;
+								}
+								PrototypePrefix = ParentPath;
+							}
+
+							// e.g. "/root/instanceable7/child_prim/slot"
+							pxr::SdfPath InstanceProxyFullPath = InstanceProxySuffix.ReplacePrefix(PrototypePrefix, InstanceRootPath);
+
+							pxr::UsdPrim InstanceRootOver = Stage->OverridePrim(InstanceRoot.GetPrimPath());
+							pxr::UsdPrim InstanceProxyPrim = Stage->GetPrimAtPath(InstanceProxyFullPath);
+							UsdUtils::AuthorUnrealCollectionBasedMaterialBinding(	 //
+								InstanceRootOver,
+								InstanceProxyPrim,
+								OverrideMaterialPath
+							);
+						}
+					}
+					else
+					{
+						// Here we can just author a regular material binding opinion
+						pxr::SdfPath OverridePrimPath = UnrealToUsd::ConvertPath(*SourcePrimPath).Get();
+						pxr::UsdPrim MeshPrim = Stage->OverridePrim(OverridePrimPath);
+						UsdUtils::AuthorUnrealMaterialBinding(MeshPrim, OverrideMaterialPath);
+					}
+				}
+			}
+		}
+	};
+
 	if (const UGeometryCache* GeometryCache = Cast<const UGeometryCache>(MeshAsset))
 	{
 		for (int32 MatIndex = 0; MatIndex < MaterialOverrides.Num(); ++MatIndex)
@@ -3406,26 +3516,13 @@ bool UnrealToUsd::ConvertMaterialOverrides(
 			}
 
 			// If we have user data this is one of our meshes, so we know exactly the prim that corresponds to each
-			// material slot. Let's use that
+			// material slot. Let's use that.
+			//
 			// const_cast as there's no const access to asset user data on IInterface_AssetUserData, but we won't
 			// modify anything
 			if (const UUsdMeshAssetUserData* UserData = const_cast<UGeometryCache*>(GeometryCache)->GetAssetUserData<UUsdMeshAssetUserData>())
 			{
-				if (const FUsdPrimPathList* SourcePrimPaths = UserData->MaterialSlotToPrimPaths.Find(MatIndex))
-				{
-					for (const FString& PrimPath : SourcePrimPaths->PrimPaths)
-					{
-						// Our mesh assets are shared between multiple prims via the asset cache, and all user prims of that asset will record
-						// their source prim paths for each material slot. In here we just want to apply overrides to the prims that
-						// correspond to the modified component, and not the others
-						if (PrimPath.StartsWith(UsdPrimPath))
-						{
-							pxr::SdfPath OverridePrimPath = UnrealToUsd::ConvertPath(*PrimPath).Get();
-							pxr::UsdPrim MeshPrim = Stage->OverridePrim(OverridePrimPath);
-							UsdUtils::AuthorUnrealMaterialBinding(MeshPrim, Override->GetPathName());
-						}
-					}
-				}
+				AuthorOverrideViaUserData(UserData, MatIndex, {MeshPrimPath}, Override->GetPathName());
 			}
 			// If we don't, we have to fallback to writing the same prim patterns that the mesh exporters
 			// generate when exporting meshes, so that we can override its opinions. This happens when exporting
@@ -3487,19 +3584,7 @@ bool UnrealToUsd::ConvertMaterialOverrides(
 
 					if (UserData)
 					{
-						if (const FUsdPrimPathList* SourcePrimPaths = UserData->MaterialSlotToPrimPaths.Find(SectionMatIndex))
-						{
-							for (const FString& PrimPath : SourcePrimPaths->PrimPaths)
-							{
-								// See comment on the analogue part for the geometry cache component
-								if (PrimPath.StartsWith(UsdPrimPath))
-								{
-									pxr::SdfPath OverridePrimPath = UnrealToUsd::ConvertPath(*PrimPath).Get();
-									pxr::UsdPrim MeshPrim = Stage->OverridePrim(OverridePrimPath);
-									UsdUtils::AuthorUnrealMaterialBinding(MeshPrim, Override->GetPathName());
-								}
-							}
-						}
+						AuthorOverrideViaUserData(UserData, MatIndex, {MeshPrimPath}, Override->GetPathName());
 					}
 					else
 					{
@@ -3609,6 +3694,11 @@ bool UnrealToUsd::ConvertMaterialOverrides(
 		for (const pxr::UsdSkelSkinningQuery& SkinningTarget : SkinningTargets)
 		{
 			pxr::UsdPrim SkinnedPrim = SkinningTarget.GetPrim();
+			if (bIsInsideInstance && SkinnedPrim.IsInstanceProxy())
+			{
+				SkinnedPrim = SkinnedPrim.GetPrimInPrototype();
+			}
+
 			if (pxr::UsdGeomMesh SkinnedMesh = pxr::UsdGeomMesh{SkinnedPrim})
 			{
 				SkinnedMeshPaths.Add(UsdToUnreal::ConvertPath(SkinnedPrim.GetPrimPath()));
@@ -3653,27 +3743,7 @@ bool UnrealToUsd::ConvertMaterialOverrides(
 
 					if (UserData)
 					{
-						if (const FUsdPrimPathList* SourcePrimPaths = UserData->MaterialSlotToPrimPaths.Find(SectionMatIndex))
-						{
-							// The N^2 here is not great but note that "ConvertMaterialOverrides" is not exactly spammed all that much, and that
-							// these arrays will in the general case have like 3 items each
-							for (const FString& SourcePrimPath : SourcePrimPaths->PrimPaths)
-							{
-								for (const FString& SkinnedMeshPath : SkinnedMeshPaths)
-								{
-									// See comment on the analogue part for the geometry cache component
-									// Note we use 'StartsWith' here because our SourcePrimPath may be a UsdGeomSubset or something like that,
-									// but only the actual Mesh prim will count as a "skinned mesh"
-									if (SourcePrimPath.StartsWith(SkinnedMeshPath))
-									{
-										pxr::SdfPath OverridePrimPath = UnrealToUsd::ConvertPath(*SourcePrimPath).Get();
-										pxr::UsdPrim MeshPrim = Stage->OverridePrim(OverridePrimPath);
-										UsdUtils::AuthorUnrealMaterialBinding(MeshPrim, Override->GetPathName());
-										break;
-									}
-								}
-							}
-						}
+						AuthorOverrideViaUserData(UserData, MatIndex, SkinnedMeshPaths, Override->GetPathName());
 					}
 					// TODO: We really need a separate function for ConvertingMaterialOverrides (to an opened stage) and ExportingMaterialOverrides
 					// that we can use when the SkeletalMesh is not something we generated ourselves (with annotated MaterialSlotToPrimPaths).
diff --git a/Engine/Plugins/Runtime/USDCore/Source/USDUtilities/Private/USDShadeConversion.cpp b/Engine/Plugins/Runtime/USDCore/Source/USDUtilities/Private/USDShadeConversion.cpp
--- a/Engine/Plugins/Runtime/USDCore/Source/USDUtilities/Private/USDShadeConversion.cpp
+++ b/Engine/Plugins/Runtime/USDCore/Source/USDUtilities/Private/USDShadeConversion.cpp
@@ -3112,6 +3112,239 @@ void UsdUtils::AuthorUnrealMaterialBinding(pxr::UsdPrim& MeshOrGeomSubsetPrim, c
 	}
 }
 
+void UsdUtils::AuthorUnrealCollectionBasedMaterialBinding(
+	const pxr::UsdPrim& CollectionPrim,
+	const pxr::UsdPrim& TargetMeshOrGeomSubsetPrim,
+	const FString& UnrealMaterialPathName
+)
+{
+	if (!CollectionPrim || CollectionPrim.IsInstanceProxy() || !TargetMeshOrGeomSubsetPrim || UnrealMaterialPathName.IsEmpty())
+	{
+		return;
+	}
+
+	FScopedUsdAllocs UsdAllocs;
+
+	// For collection-based bindings to do anything, TargetMeshOrGeomSubsetPrim must be a descendant of CollectionPrim
+	pxr::UsdPrim Iter = TargetMeshOrGeomSubsetPrim.GetParent();
+	while (Iter != CollectionPrim)
+	{
+		if (!Iter || Iter.IsPseudoRoot())
+		{
+			UE_LOG(
+				LogUsd,
+				Warning,
+				TEXT(
+					"Failed to author collection-based material bindings on prim '%s' for target prim '%s', as the latter is not a descendent of the former"
+				),
+				*UsdToUnreal::ConvertPath(CollectionPrim.GetPath()),
+				*UsdToUnreal::ConvertPath(TargetMeshOrGeomSubsetPrim.GetPath())
+			);
+			return;
+		}
+		Iter = Iter.GetParent();
+	}
+
+	pxr::UsdStageRefPtr Stage = CollectionPrim.GetStage();
+
+	const static pxr::TfToken& BindingPurpose = pxr::UsdShadeTokens->allPurpose;
+	pxr::UsdShadeMaterialBindingAPI BindingAPI = pxr::UsdShadeMaterialBindingAPI::Apply(CollectionPrim);
+	std::vector<pxr::UsdShadeMaterialBindingAPI::CollectionBinding> ExistingBindings = BindingAPI.GetCollectionBindings(BindingPurpose);
+
+	// Check to see if we happen to have a collection based binding for our target material already. This can happen
+	// if we're e.g. setting multiple prims inside of an instance with the same material override
+	pxr::UsdShadeMaterialBindingAPI::CollectionBinding CollectionBinding;
+	pxr::UsdShadeMaterial BoundMaterial;
+	{
+		for (const pxr::UsdShadeMaterialBindingAPI::CollectionBinding& Binding : ExistingBindings)
+		{
+			pxr::UsdShadeMaterial Material = Binding.GetMaterial();
+
+			TOptional<FString> ExistingUnrealMaterialPath = UsdUtils::GetUnrealSurfaceOutput(Material.GetPrim());
+			if (ExistingUnrealMaterialPath.IsSet() && ExistingUnrealMaterialPath.GetValue() == UnrealMaterialPathName)
+			{
+				CollectionBinding = Binding;
+				BoundMaterial = CollectionBinding.GetMaterial();
+				break;
+			}
+		}
+	}
+
+	// If the CollectionPrim is an instance, we won't be able to author any material inside of it and must instead
+	// settle for creating sibling Material prims
+	pxr::UsdPrim MaterialParent = CollectionPrim.IsInstance() ? CollectionPrim.GetParent() : CollectionPrim;
+
+	// We don't have a collection-based material binding to our target material yet.
+	// Double-check we don't have any existing child Material on this CollectionPrim that we can just reuse though
+	if (!BoundMaterial)
+	{
+		for (const pxr::UsdPrim& Child : MaterialParent.GetFilteredChildren(pxr::UsdTraverseInstanceProxies(pxr::UsdPrimAllPrimsPredicate)))
+		{
+			TOptional<FString> ExistingUnrealMaterialPath = UsdUtils::GetUnrealSurfaceOutput(Child.GetPrim());
+			if (ExistingUnrealMaterialPath.IsSet() && ExistingUnrealMaterialPath.GetValue() == UnrealMaterialPathName)
+			{
+				BoundMaterial = pxr::UsdShadeMaterial{Child};
+				break;
+			}
+		}
+	}
+
+	// Need to create a brand new material with an "unreal" surface output that just points at our target material
+	if (!BoundMaterial)
+	{
+		// Find a unique name for our child material prim
+		// Note how we'll always author these materials as children of the meshes themselves instead of emitting a common
+		// Material prim to use for multiple overrides: This because in the future we'll want to have a separate material
+		// bake for each mesh (to make sure we get vertex color effects, etc.), and so we'd have multiple baked .usda material
+		// asset layers for each UE material, and we'd want each mesh/section/LOD to refer to its own anyway
+		FString ChildMaterialName = TEXT("UnrealMaterial");
+		if (pxr::UsdPrim ExistingPrim = MaterialParent.GetChild(UnrealToUsd::ConvertToken(*ChildMaterialName).Get()))
+		{
+			// Get a unique name for a new prim. Don't even try checking if this prim is usable as the material binding,
+			// because if it was we would have already found it above
+			TSet<FString> UsedNames;
+			for (pxr::UsdPrim Child : MaterialParent.GetFilteredChildren(pxr::UsdTraverseInstanceProxies(pxr::UsdPrimAllPrimsPredicate)))
+			{
+				UsedNames.Add(UsdToUnreal::ConvertToken(Child.GetName()));
+			}
+
+			ChildMaterialName = UsdUnreal::ObjectUtils::GetUniqueName(ChildMaterialName, UsedNames);
+		}
+
+		pxr::SdfPath MaterialParentPath = MaterialParent.GetPath();
+		pxr::SdfPath MaterialPath = MaterialParentPath.AppendChild(UnrealToUsd::ConvertToken(*ChildMaterialName).Get());
+
+		BoundMaterial = pxr::UsdShadeMaterial::Define(Stage, MaterialPath);
+		if (!BoundMaterial)
+		{
+			UE_LOG(
+				LogUsd,
+				Warning,
+				TEXT("Failed to author material prim '%s' when trying to write '%s's collection-based material assignment '%s' to USD"),
+				*UsdToUnreal::ConvertPath(MaterialPath),
+				*UsdToUnreal::ConvertPath(CollectionPrim.GetPath()),
+				*UnrealMaterialPathName
+			);
+			return;
+		}
+
+		if (pxr::UsdPrim MaterialPrim = BoundMaterial.GetPrim())
+		{
+			UsdUtils::SetUnrealSurfaceOutput(MaterialPrim, UnrealMaterialPathName);
+		}
+	}
+
+	const static pxr::TfToken& BindingStrength = pxr::UsdShadeTokens->strongerThanDescendants;
+
+	// Get the target collection name
+	pxr::TfToken CollectionName;
+	if (CollectionBinding.IsValid())
+	{
+		// We already have a collection based binding to our target material, let's just use it
+		CollectionName = CollectionBinding.GetCollection().GetName();
+	}
+	else
+	{
+		// We don't have a collection binding to our target material yet
+		pxr::UsdCollectionAPI CollectionToUse;
+
+		// Before we try creating a brand new collection, let's see if our target prim is already the single prim
+		// targetted by other existing collections, because if it is we could just make that collection point at
+		// our new material instead.
+		//
+		// This is useful because otherwise swapping between 3 different material overrides on the exact same
+		// component would generate 3 separate collections on the prim, author a bunch of "delete" opinions, and
+		// overall just make a mess. With this snippet we just update the material on the same collection instead
+		for (const pxr::UsdShadeMaterialBindingAPI::CollectionBinding& Binding : ExistingBindings)
+		{
+			pxr::TfToken ExistingStrength = BindingAPI.GetMaterialBindingStrength(Binding.GetBindingRel());
+			if (ExistingStrength != BindingStrength)
+			{
+				// Only pick collections that looks like the one we'd create though
+				continue;
+			}
+
+			pxr::UsdCollectionAPI ExistingCollection = Binding.GetCollection();
+			std::set<pxr::SdfPath> IncludedPaths = ExistingCollection.ComputeIncludedPaths(ExistingCollection.ComputeMembershipQuery(), Stage);
+			if (IncludedPaths.size() == 1 && (*IncludedPaths.begin()) == TargetMeshOrGeomSubsetPrim.GetPrimPath())
+			{
+				CollectionToUse = ExistingCollection;
+				break;
+			}
+		}
+
+		// We need to create a brand new collection
+		if (!CollectionToUse)
+		{
+			// Find a unique name for our new collection
+			pxr::TfToken NewCollectionName;
+			{
+				TSet<FString> UsedNames;
+				for (const pxr::UsdShadeMaterialBindingAPI::CollectionBinding& Binding : ExistingBindings)
+				{
+					UsedNames.Add(UsdToUnreal::ConvertToken(Binding.GetCollection().GetName()));
+				}
+
+				const static FString UnrealOverridesStr = TEXT("unrealOverrides");
+				FString NewCollectionNameStr = UsdUnreal::ObjectUtils::GetUniqueName(UnrealOverridesStr, UsedNames);
+				NewCollectionName = UnrealToUsd::ConvertToken(*NewCollectionNameStr).Get();
+			}
+
+			// Actually create the new collection binding, with our new name and material
+			CollectionToUse = pxr::UsdCollectionAPI::Apply(CollectionPrim, NewCollectionName);
+		}
+
+		// Set our target material on CollectionToUse
+		pxr::TfToken Empty;	   // Will use CollectionToUse's name instead
+		const bool bCreatedBinding = BindingAPI.Bind(CollectionToUse, BoundMaterial, Empty, BindingStrength);
+		if (!bCreatedBinding)
+		{
+			UE_LOG(
+				LogUsd,
+				Warning,
+				TEXT("Failed to create collection-based material binding '%s' on prim '%s'"),
+				*UsdToUnreal::ConvertToken(CollectionToUse.GetName()),
+				*UsdToUnreal::ConvertPath(CollectionPrim.GetPath())
+			);
+			return;
+		}
+
+		CollectionName = CollectionToUse.GetName();
+	}
+
+	// Add the prim to the collection with CollectionName
+	const bool bAddedPrim = BindingAPI.AddPrimToBindingCollection(TargetMeshOrGeomSubsetPrim, CollectionName, BindingPurpose);
+	if (!bAddedPrim)
+	{
+		UE_LOG(
+			LogUsd,
+			Warning,
+			TEXT("Failed to add prim '%s' to the collection-based material binding '%s' on prim '%s'"),
+			*UsdToUnreal::ConvertPath(TargetMeshOrGeomSubsetPrim.GetPath()),
+			*UsdToUnreal::ConvertToken(CollectionName),
+			*UsdToUnreal::ConvertPath(CollectionPrim.GetPath())
+		);
+		return;
+	}
+
+	// Remove that prim from any other collection-based bindings we may already have on this prim, to make sure the override shows up
+	if (bAddedPrim)
+	{
+		std::vector<pxr::UsdShadeMaterialBindingAPI::CollectionBinding> AllCollections = BindingAPI.GetCollectionBindings(BindingPurpose);
+		for (const pxr::UsdShadeMaterialBindingAPI::CollectionBinding& SomeCollectionBinding : AllCollections)
+		{
+			pxr::UsdCollectionAPI Collection = SomeCollectionBinding.GetCollection();
+			pxr::TfToken OtherCollectionName = Collection.GetName();
+			if (OtherCollectionName == CollectionName)
+			{
+				continue;
+			}
+
+			BindingAPI.RemovePrimFromBindingCollection(TargetMeshOrGeomSubsetPrim, OtherCollectionName, BindingPurpose);
+		}
+	}
+}
+
 TOptional<FString> UsdUtils::GetUnrealSurfaceOutput(const pxr::UsdPrim& MaterialPrim)
 {
 	TRACE_CPUPROFILER_EVENT_SCOPE(UsdUtils::GetUnrealSurfaceOutput);
diff --git a/Engine/Plugins/Runtime/USDCore/Source/USDUtilities/Public/USDShadeConversion.h b/Engine/Plugins/Runtime/USDCore/Source/USDUtilities/Public/USDShadeConversion.h
--- a/Engine/Plugins/Runtime/USDCore/Source/USDUtilities/Public/USDShadeConversion.h
+++ b/Engine/Plugins/Runtime/USDCore/Source/USDUtilities/Public/USDShadeConversion.h
@@ -49,7 +49,7 @@ namespace UsdToUnreal
 {
 	struct USDUTILITIES_API FTextureParameterValue
 	{
-		UTexture* Texture = nullptr;	  // Only used for the ConvertMaterial overloads that receive TexturesCache
+		UTexture* Texture = nullptr;	// Only used for the ConvertMaterial overloads that receive TexturesCache
 
 		// Parameters of the texture asset itself
 		FString TextureFilePath;
@@ -230,6 +230,23 @@ namespace UsdUtils
 	// Doesn't write to the 'unrealMaterial' attribute at all, as we intend on deprecating it in the future.
 	USDUTILITIES_API void AuthorUnrealMaterialBinding(pxr::UsdPrim& MeshOrGeomSubsetPrim, const FString& UnrealMaterialPathName);
 
+	/**
+	 * Similar to AuthorUnrealMaterialBinding, but instead of authoring material bindings directly to TargetMeshOrGeomSubsetPrim,
+	 * it will instead author collection-based material bindings on the CollectionPrim, that instead target TargetMeshOrGeomSubsetPrim.
+	 *
+	 * It will try reusing existing collections and UnrealMaterials, but otherwise it will author a new collection within CollectionPrim,
+	 * and a new UnrealMaterial as a sibling of CollectionPrim, referring to UnrealMaterialPathName.
+	 *
+	 * This is useful if CollectionPrim is an instance root, and TargetMeshOrGeomSubsetPrim is an instance proxy, for example.
+	 *
+	 * WARNING: In order to get collection-based bindings to work, TargetMeshOrGeomSubsetPrim must be a descendant of CollectionPrim.
+	 */
+	USDUTILITIES_API void AuthorUnrealCollectionBasedMaterialBinding(
+		const pxr::UsdPrim& CollectionPrim,
+		const pxr::UsdPrim& TargetMeshOrGeomSubsetPrim,
+		const FString& UnrealMaterialPathName
+	);
+
 	/** Returns a path to an UE asset (e.g. "/Game/Assets/Red.Red") if MaterialPrim has an 'unreal' render context surface output that points at one
 	 */
 	USDUTILITIES_API TOptional<FString> GetUnrealSurfaceOutput(const pxr::UsdPrim& MaterialPrim);
