USD: Add the bSkipConstantValues option for exporting LevelSequences to USD.

#virtualized

diff --git a/Engine/Plugins/Importers/USDImporter/Source/USDExporter/Private/LevelSequenceExporterUSD.cpp b/Engine/Plugins/Importers/USDImporter/Source/USDExporter/Private/LevelSequenceExporterUSD.cpp
--- a/Engine/Plugins/Importers/USDImporter/Source/USDExporter/Private/LevelSequenceExporterUSD.cpp
+++ b/Engine/Plugins/Importers/USDImporter/Source/USDExporter/Private/LevelSequenceExporterUSD.cpp
@@ -528,6 +528,8 @@ namespace UE::LevelSequenceExporterUSD::Private
 		UMovieSceneSequence& RootSequence
 	)
 	{
+		TRACE_CPUPROFILER_EVENT_SCOPE(ULevelSequenceExporterUsd::PreSpawnSpawnables);
+
 		FMovieSceneSequenceHierarchy HierarchyCache;
 		TMap<UMovieSceneSequence*, TArray<FMovieSceneSequenceID>> SequenceInstances = GetSequenceHierarchyInstances(
 			RootSequence,
@@ -835,6 +837,8 @@ namespace UE::LevelSequenceExporterUSD::Private
 		TMap<USceneComponent*, FCombinedComponentBakers>& InOutComponentBakers
 	)
 	{
+		TRACE_CPUPROFILER_EVENT_SCOPE(ULevelSequenceExporterUsd::GenerateBakersForMovieScene);
+
 		UMovieScene* MovieScene = MovieSceneSequence.GetMovieScene();
 		if (!MovieScene)
 		{
@@ -1069,9 +1073,9 @@ namespace UE::LevelSequenceExporterUSD::Private
 									// Note that it's perfectly fine if this ends up being just "."
 									UE::FSdfPath RelativePrimPathInSourceStage = SdfPrimPathInSourceStage.MakeRelativePath(DefaultPrimPath);
 
-									FString PrimPathRelativeToParentStageActor = UE::FSdfPath{ *ParentStageActorPrimPathOnExport }
-										.AppendPath(RelativePrimPathInSourceStage)
-										.GetString();
+									FString PrimPathRelativeToParentStageActor = UE::FSdfPath{*ParentStageActorPrimPathOnExport}
+																					 .AppendPath(RelativePrimPathInSourceStage)
+																					 .GetString();
 
 									if (!PrimPathRelativeToParentStageActor.IsEmpty())
 									{
@@ -1101,11 +1105,12 @@ namespace UE::LevelSequenceExporterUSD::Private
 
 			auto GetPrimForComponent = [&Context, &UsdStage](const USceneComponent& Component, FString* PrimPathStr = nullptr) -> UE::FUsdPrim
 			{
-				FString PrimPath = PrimPathStr ? *PrimPathStr : UsdUtils::GetPrimPathForObject(
-					&Component,
-					TEXT(""),
-					Context.ExportOptions && Context.ExportOptions->LevelExportOptions.bExportActorFolders
-				);
+				FString PrimPath = PrimPathStr ? *PrimPathStr
+											   : UsdUtils::GetPrimPathForObject(
+												   &Component,
+												   TEXT(""),
+												   Context.ExportOptions && Context.ExportOptions->LevelExportOptions.bExportActorFolders
+											   );
 				if (PrimPath.IsEmpty())
 				{
 					return {};
@@ -1119,7 +1124,7 @@ namespace UE::LevelSequenceExporterUSD::Private
 
 				// We will define a prim here so that we can apply schemas and use the shortcut CreateXAttribute functions,
 				// and not have to worry about attribute names and types. Later on we will convert these prims back into just 'overs' though
-				return UsdStage.DefinePrim(UE::FSdfPath{ *PrimPath }, *SchemaName);
+				return UsdStage.DefinePrim(UE::FSdfPath{*PrimPath}, *SchemaName);
 			};
 
 			UE::FUsdPrim Prim = GetPrimForComponent(*BoundComponent, &PrimPath);
@@ -1129,55 +1134,60 @@ namespace UE::LevelSequenceExporterUSD::Private
 			}
 
 			TFunction<void(UnrealToUsd::FComponentBaker&)> AddBaker = [&InOutComponentBakers, &BoundComponent](UnrealToUsd::FComponentBaker& Baker)
+			{
+				// If we made a baker and we don't have one of this type for this component yet, add its lambda to the array
+				FCombinedComponentBakers& ExistingBakers = InOutComponentBakers.FindOrAdd(BoundComponent);
+				if (Baker.BakerType != UnrealToUsd::EBakingType::None && !EnumHasAnyFlags(ExistingBakers.CombinedBakingType, Baker.BakerType))
 				{
-					// If we made a baker and we don't have one of this type for this component yet, add its lambda to the array
-					FCombinedComponentBakers& ExistingBakers = InOutComponentBakers.FindOrAdd(BoundComponent);
-					if (Baker.BakerType != UnrealToUsd::EBakingType::None && !EnumHasAnyFlags(ExistingBakers.CombinedBakingType, Baker.BakerType))
-					{
-						ExistingBakers.Bakers.Add(Baker);
-						ExistingBakers.CombinedBakingType |= Baker.BakerType;
-					}
-				};
+					ExistingBakers.Bakers.Add(Baker);
+					ExistingBakers.CombinedBakingType |= Baker.BakerType;
+				}
+			};
 
 			const bool bBakeAsSkeletal = !Context.ExportOptions->LevelExportOptions.AssetOptions.bConvertSkeletalToNonSkeletal;
 			TFunction<void(UnrealToUsd::FComponentBaker&)> GenerateSkeletalBaker =
 				[&UsdStage, &BoundComponent, &PrimPath, &Prim, bBakeAsSkeletal](UnrealToUsd::FComponentBaker& InOutBaker)
+			{
+				if (USkeletalMeshComponent* SkeletalBoundComponent = Cast<USkeletalMeshComponent>(BoundComponent))
 				{
-					if (USkeletalMeshComponent* SkeletalBoundComponent = Cast<USkeletalMeshComponent>(BoundComponent))
+					if (bBakeAsSkeletal)
 					{
-						if (bBakeAsSkeletal)
-						{
-							UE::FUsdPrim SkelAnimPrim = UsdStage.DefinePrim(UE::FSdfPath{ *PrimPath }.AppendChild(TEXT("Anim")), TEXT("SkelAnimation"));
+						UE::FUsdPrim SkelAnimPrim = UsdStage.DefinePrim(UE::FSdfPath{*PrimPath}.AppendChild(TEXT("Anim")), TEXT("SkelAnimation"));
 
-							UE::FUsdPrim SkeletonPrim = UsdStage.DefinePrim(
-								UE::FSdfPath{ *PrimPath }.AppendChild(UnrealIdentifiers::ExportedSkeletonPrimName),
-								TEXT("Skeleton")
-							);
+						UE::FUsdPrim SkeletonPrim = UsdStage.DefinePrim(
+							UE::FSdfPath{*PrimPath}.AppendChild(UnrealIdentifiers::ExportedSkeletonPrimName),
+							TEXT("Skeleton")
+						);
 
-							if (SkelAnimPrim && SkeletonPrim)
-							{
-								UnrealToUsd::CreateSkeletalAnimationBaker(SkeletonPrim, SkelAnimPrim, *SkeletalBoundComponent, InOutBaker);
-							}
-							else
-							{
-								UE_LOG(LogUsd, Warning, TEXT("Failed to generate Skeleton or SkelAnimation prim when baking out SkelRoot '%s'"), *PrimPath);
-							}
+						if (SkelAnimPrim && SkeletonPrim)
+						{
+							UnrealToUsd::CreateSkeletalAnimationBaker(SkeletonPrim, SkelAnimPrim, *SkeletalBoundComponent, InOutBaker);
 						}
 						else
 						{
-							// Convert the prim for the skeletal mesh component from SkelRoot to Mesh
-							Prim.SetTypeName(TEXT("Mesh"));
-							UnrealToUsd::CreateSkeletalAnimationToMeshBaker(Prim, *SkeletalBoundComponent, InOutBaker);
+							UE_LOG(
+								LogUsd,
+								Warning,
+								TEXT("Failed to generate Skeleton or SkelAnimation prim when baking out SkelRoot '%s'"),
+								*PrimPath
+							);
 						}
 					}
-				};
+					else
+					{
+						// Convert the prim for the skeletal mesh component from SkelRoot to Mesh
+						Prim.SetTypeName(TEXT("Mesh"));
+						UnrealToUsd::CreateSkeletalAnimationToMeshBaker(Prim, *SkeletalBoundComponent, InOutBaker);
+					}
+				}
+			};
 
 			bool bHasTransformBaker = false;
 			bool bHasSkeletalBaker = false;
 
 			if (ACameraRig_Rail* RailActor = Cast<ACameraRig_Rail>(BoundObject))
 			{
-				// In the case of a CameraRig_Rail, what we want to bake is the transform animation of its 
+				// In the case of a CameraRig_Rail, what we want to bake is the transform animation of its
 				// RailMountComponent since that is where children camera will be attached to.
 				const static FString TransformPropertyPath = UnrealIdentifiers::TransformPropertyName.ToString();
 				USceneComponent* RailMountComponent = RailActor->GetDefaultAttachComponent();
@@ -1377,6 +1387,8 @@ namespace UE::LevelSequenceExporterUSD::Private
 		const TMap<USceneComponent*, FCombinedComponentBakers>& ComponentBakers
 	)
 	{
+		TRACE_CPUPROFILER_EVENT_SCOPE(ULevelSequenceExporterUsd::BakeMovieSceneSequence);
+
 		UMovieScene* MovieScene = MovieSceneSequence.GetMovieScene();
 		if (!MovieScene)
 		{
@@ -1453,21 +1465,25 @@ namespace UE::LevelSequenceExporterUSD::Private
 			}
 		);
 
-		for (FFrameTime EvalTime = StartFrame; EvalTime <= EndFrame; EvalTime += Interval)
 		{
-			Context.Sequencer->SetLocalTimeDirectly(EvalTime);
-			Context.Sequencer->ForceEvaluate();
+			TRACE_CPUPROFILER_EVENT_SCOPE(PlaySequence);
 
-			// Evaluate constraints (these run on tick in the editor, so here we must trigger them manually)
-			// Can't iterate through a pre-sorted list since the parenting of the constraints can change between frames
-			Controller.EvaluateAllConstraints();
+			for (FFrameTime EvalTime = StartFrame; EvalTime <= EndFrame; EvalTime += Interval)
+			{
+				Context.Sequencer->SetLocalTimeDirectly(EvalTime);
+				Context.Sequencer->ForceEvaluate();
 
-			FFrameTime KeyTime = FFrameRate::Snap(EvalTime, Resolution, DisplayRate).FloorToFrame();
-			double UsdTimeCode = FFrameRate::TransformTime(KeyTime, Resolution, StageFrameRate).AsDecimal();
+				// Evaluate constraints (these run on tick in the editor, so here we must trigger them manually)
+				// Can't iterate through a pre-sorted list since the parenting of the constraints can change between frames
+				Controller.EvaluateAllConstraints();
 
-			for (const UnrealToUsd::FComponentBaker& Baker : SortedBakers)
-			{
-				Baker.BakerFunction(UsdTimeCode);
+				FFrameTime KeyTime = FFrameRate::Snap(EvalTime, Resolution, DisplayRate).FloorToFrame();
+				double UsdTimeCode = FFrameRate::TransformTime(KeyTime, Resolution, StageFrameRate).AsDecimal();
+
+				for (const UnrealToUsd::FComponentBaker& Baker : SortedBakers)
+				{
+					Baker.BakerFunction(UsdTimeCode);
+				}
 			}
 		}
 
@@ -1499,6 +1515,8 @@ namespace UE::LevelSequenceExporterUSD::Private
 		TArray<UE::FUsdStage>& InOutExportedStages
 	)
 	{
+		TRACE_CPUPROFILER_EVENT_SCOPE(ULevelSequenceExporterUsd::ExportMovieSceneSequence);
+
 		if (FilePath.IsEmpty() || Context.ExportedMovieScenes.Contains(&MovieSceneSequence))
 		{
 			return;
@@ -1728,7 +1746,10 @@ namespace UE::LevelSequenceExporterUSD::Private
 		Context.ExportedMovieScenes.Add(&MovieSceneSequence, UniqueFilePath);
 		Context.UsedFilePaths.Add(UniqueFilePath);
 
-		UsdStage.GetRootLayer().Save();
+		{
+			TRACE_CPUPROFILER_EVENT_SCOPE(Save);
+			UsdStage.GetRootLayer().Save();
+		}
 
 		InOutExportedStages.Add(UsdStage);
 	}
@@ -1755,6 +1776,8 @@ bool ULevelSequenceExporterUsd::ExportBinary(
 {
 	namespace LevelSequenceExporterImpl = UE::LevelSequenceExporterUSD::Private;
 
+	TRACE_CPUPROFILER_EVENT_SCOPE(ULevelSequenceExporterUsd::ExportBinary);
+
 #if USE_USD_SDK
 	ULevelSequence* LevelSequence = Cast<ULevelSequence>(Object);
 	if (!GEditor || !GIsEditor || !LevelSequence)
@@ -1927,6 +1950,8 @@ bool ULevelSequenceExporterUsd::ExportBinary(
 	{
 		if (UWorld* WorldToExport = Options->Level.Get())
 		{
+			TRACE_CPUPROFILER_EVENT_SCOPE(ULevelSequenceExporterUsd::LevelExport);
+
 			// Come up with a file path for the level
 			FString Directory;
 			FString Filename;
diff --git a/Engine/Plugins/Importers/USDImporter/Source/USDExporter/Private/LevelSequenceExporterUSDOptions.cpp b/Engine/Plugins/Importers/USDImporter/Source/USDExporter/Private/LevelSequenceExporterUSDOptions.cpp
--- a/Engine/Plugins/Importers/USDImporter/Source/USDExporter/Private/LevelSequenceExporterUSDOptions.cpp
+++ b/Engine/Plugins/Importers/USDImporter/Source/USDExporter/Private/LevelSequenceExporterUSDOptions.cpp
@@ -6,6 +6,7 @@
 #include "USDExporterModule.h"
 
 #include "AnalyticsEventAttribute.h"
+#include "HAL/IConsoleManager.h"
 
 void UsdUtils::AddAnalyticsAttributes(const ULevelSequenceExporterUsdOptions& Options, TArray<struct FAnalyticsEventAttribute>& InOutAttributes)
 {
@@ -23,6 +24,7 @@ void UsdUtils::AddAnalyticsAttributes(const ULevelSequenceExporterUsdOptions& Op
 	{
 		InOutAttributes.Emplace(TEXT("UseExportedLevelAsSublayer"), Options.bUseExportedLevelAsSublayer);
 	}
+	InOutAttributes.Emplace(TEXT("ExportSeparatePrimsPerSpawnableInstance"), Options.bExportSeparatePrimsPerSpawnableInstance);
 	InOutAttributes.Emplace(TEXT("ReExportIdenticalLevelsAndSequences"), Options.bReExportIdenticalLevelsAndSequences);
 	InOutAttributes.Emplace(TEXT("ReExportIdenticalAssets"), Options.bReExportIdenticalAssets);
 	if (Options.bExportLevel)
@@ -52,4 +54,13 @@ void UsdUtils::HashForLevelSequenceExport(const ULevelSequenceExporterUsdOptions
 
 	const bool bUsingLevelSublayer = Options.bExportLevel && Options.bUseExportedLevelAsSublayer;
 	HashToUpdate.Update(reinterpret_cast<const uint8*>(&bUsingLevelSublayer), sizeof(bUsingLevelSublayer));
+
+	HashToUpdate.Update(
+		reinterpret_cast<const uint8*>(&Options.bExportSeparatePrimsPerSpawnableInstance),
+		sizeof(Options.bExportSeparatePrimsPerSpawnableInstance)
+	);
+
+	static IConsoleVariable* Cvar = IConsoleManager::Get().FindConsoleVariable(TEXT("USD.LevelSequenceExport.SkipConstantValues"));
+	const bool bSkipConstantValues = Cvar && Cvar->GetBool();
+	HashToUpdate.Update(reinterpret_cast<const uint8*>(&bSkipConstantValues), sizeof(bSkipConstantValues));
 }
diff --git a/Engine/Plugins/Runtime/USDCore/Source/USDUtilities/Private/USDGeomMeshConversion.cpp b/Engine/Plugins/Runtime/USDCore/Source/USDUtilities/Private/USDGeomMeshConversion.cpp
--- a/Engine/Plugins/Runtime/USDCore/Source/USDUtilities/Private/USDGeomMeshConversion.cpp
+++ b/Engine/Plugins/Runtime/USDCore/Source/USDUtilities/Private/USDGeomMeshConversion.cpp
@@ -4205,7 +4205,8 @@ namespace UE::UsdGeometryCacheConversion::Private
 
 		FGeometryCacheExportContext(const TArray<FName>& InSlotNames)
 			: SlotNames(InSlotNames)
-		{}
+		{
+		}
 
 		TArray<FName> SlotNames;
 		int32 InclusiveEndFrame;
@@ -4421,7 +4422,11 @@ namespace UE::UsdGeometryCacheConversion::Private
 	class FSkinnedVerticesDataWrapper : public FSkeletalMeshDataWrapper
 	{
 	public:
-		FSkinnedVerticesDataWrapper(const TArray<FFinalSkinVertex>& InSkinnedVertices, const FSkeletalMeshLODRenderData& InLODData, bool bVertexColors)
+		FSkinnedVerticesDataWrapper(
+			const TArray<FFinalSkinVertex>& InSkinnedVertices,
+			const FSkeletalMeshLODRenderData& InLODData,
+			bool bVertexColors
+		)
 			: FSkeletalMeshDataWrapper(InLODData, bVertexColors)
 			, SkinnedVertices(InSkinnedVertices)
 		{
@@ -4449,7 +4454,7 @@ namespace UE::UsdGeometryCacheConversion::Private
 			return false;
 		}
 		// In case we want to use the UV from the skinned vertices, but they are not animated anyway
-		//FVector2D GetUV(int32 Index) const override
+		// FVector2D GetUV(int32 Index) const override
 		//{
 		//	return SkinnedVertices[Index].TextureCoordinates[0];
 		//}
@@ -4545,7 +4550,7 @@ namespace UE::UsdGeometryCacheConversion::Private
 				pxr::TfToken UsdUVSetName = UsdUtils::GetUVSetName(TexCoordSourceIndex).Get();
 
 				pxr::UsdGeomPrimvar PrimvarST = pxr::UsdGeomPrimvarsAPI(MeshPrim)
-					.CreatePrimvar(UsdUVSetName, pxr::SdfValueTypeNames->TexCoord2fArray, pxr::UsdGeomTokens->vertex);
+													.CreatePrimvar(UsdUVSetName, pxr::SdfValueTypeNames->TexCoord2fArray, pxr::UsdGeomTokens->vertex);
 
 				if (PrimvarST)
 				{
@@ -4707,7 +4712,7 @@ namespace UE::UsdGeometryCacheConversion::Private
 						);
 					}
 
-					pxr::UsdGeomSubset GeomSubsetSchema{ GeomSubsetPrim };
+					pxr::UsdGeomSubset GeomSubsetSchema{GeomSubsetPrim};
 
 					// Element type attribute
 					// Write the geomsubset attributes only once since they are at Default time anyway
@@ -4718,7 +4723,8 @@ namespace UE::UsdGeometryCacheConversion::Private
 
 						// Indices attribute
 						const uint32 TriangleCount = MeshData.GetSectionNumTriangles(SectionIndex);
-						const uint32 FirstTriangleIndex = MeshData.GetSectionStartIndex(SectionIndex) / 3;	 // StartIndex is the first *vertex* instance index
+						const uint32 FirstTriangleIndex = MeshData.GetSectionStartIndex(SectionIndex) / 3;	  // StartIndex is the first *vertex*
+																											  // instance index
 						pxr::VtArray<int> IndicesAttrValue;
 						for (uint32 TriangleIndex = FirstTriangleIndex; TriangleIndex - FirstTriangleIndex < TriangleCount; ++TriangleIndex)
 						{
@@ -4912,7 +4918,14 @@ bool UnrealToUsd::ConvertSkeletalMeshToStaticMesh(
 
 	FGeometryCacheExportContext ExportContext(SlotNames);
 	pxr::UsdPrim MaterialPrim = MaterialStage->OverridePrim(UsdPrim.GetPath());
-	ConvertMeshData(FSkeletalMeshDataWrapper(LODData, SkeletalMesh->GetHasVertexColors()), MaterialAssignments, TimeCode, MaterialPrim, ExportContext, MeshPrim);
+	ConvertMeshData(
+		FSkeletalMeshDataWrapper(LODData, SkeletalMesh->GetHasVertexColors()),
+		MaterialAssignments,
+		TimeCode,
+		MaterialPrim,
+		ExportContext,
+		MeshPrim
+	);
 
 	pxr::UsdAttribute ExtentsAttr = MeshPrim.CreateExtentAttr();
 	const FBox BoundingBox = SkeletalMesh->GetBounds().GetBox();
@@ -5030,7 +5043,14 @@ bool UnrealToUsd::ConvertAnimSequenceToAnimatedMesh(
 		TArray<FFinalSkinVertex> SkinnedVertices;
 		SkelMeshComponent->GetCPUSkinnedVertices(SkinnedVertices, LODLevel);
 
-		ConvertMeshData(FSkinnedVerticesDataWrapper(SkinnedVertices, LODData, bHasVertexColors), MaterialAssignments, FrameIndex, MaterialPrim, ExportContext, MeshPrim);
+		ConvertMeshData(
+			FSkinnedVerticesDataWrapper(SkinnedVertices, LODData, bHasVertexColors),
+			MaterialAssignments,
+			FrameIndex,
+			MaterialPrim,
+			ExportContext,
+			MeshPrim
+		);
 	}
 
 	World->DestroyActor(SkelMeshActor);
@@ -5038,7 +5058,11 @@ bool UnrealToUsd::ConvertAnimSequenceToAnimatedMesh(
 	return true;
 }
 
-bool UnrealToUsd::CreateSkeletalAnimationToMeshBaker(UE::FUsdPrim& UsdPrim, USkeletalMeshComponent& SkelMeshComponent, UnrealToUsd::FComponentBaker& OutBaker)
+bool UnrealToUsd::CreateSkeletalAnimationToMeshBaker(
+	UE::FUsdPrim& UsdPrim,
+	USkeletalMeshComponent& SkelMeshComponent,
+	UnrealToUsd::FComponentBaker& OutBaker
+)
 {
 	using namespace UE::UsdGeometryCacheConversion::Private;
 
@@ -5092,8 +5116,19 @@ bool UnrealToUsd::CreateSkeletalAnimationToMeshBaker(UE::FUsdPrim& UsdPrim, USke
 	const int32 LODLevel = 0;
 	const FSkeletalMeshLODRenderData& LODData = SkelMeshRenderData.LODRenderData[LODLevel];
 
+	static IConsoleVariable* SkipConstantValuesCvar = IConsoleManager::Get().FindConsoleVariable(TEXT("USD.LevelSequenceExport.SkipConstantValues"));
+	const bool bSkipConstantValues = SkipConstantValuesCvar && SkipConstantValuesCvar->GetBool();
+
 	const bool bHasVertexColors = SkeletalMesh->GetHasVertexColors();
-	OutBaker.BakerFunction = [&SkelMeshComponent, &LODData, MaterialAssignments, UsdPrim, bHasVertexColors, ExportContext = FGeometryCacheExportContext(SlotNames)](double UsdTimeCode) mutable
+	OutBaker.BakerFunction = [&SkelMeshComponent,
+							  &LODData,
+							  MaterialAssignments,
+							  UsdPrim,
+							  bHasVertexColors,
+							  ExportContext = FGeometryCacheExportContext(SlotNames),
+							  LastValue = TOptional<TArray<FFinalSkinVertex>>(),
+							  LastTimeCode = -DBL_MAX,
+							  bSkipConstantValues](double UsdTimeCode) mutable
 	{
 		FScopedUsdAllocs InnerAllocs;
 
@@ -5112,8 +5147,63 @@ bool UnrealToUsd::CreateSkeletalAnimationToMeshBaker(UE::FUsdPrim& UsdPrim, USke
 		TArray<FFinalSkinVertex> SkinnedVertices;
 		SkelMeshComponent.GetCPUSkinnedVertices(SkinnedVertices, LODLevel);
 
+		// This logic is copied over from the CreateCachedAttrSetter functions on USDPrimConversion.cpp so we don't have
+		// to expose neither CreateCachedAttrSetter nor the ConvertMeshData/IUnrealMeshData
+		bool bNewValueIsEqual = false;
+		if (LastValue.IsSet())
+		{
+			if (LastValue.GetValue().Num() != SkinnedVertices.Num())
+			{
+				bNewValueIsEqual = false;
+			}
+			else
+			{
+				bNewValueIsEqual = true;
+				for (int32 Index = 0; Index < SkinnedVertices.Num(); ++Index)
+				{
+					const FFinalSkinVertex& LastValueEntry = LastValue.GetValue()[Index];
+					const FFinalSkinVertex& NewValueEntry = SkinnedVertices[Index];
+
+					// Only checking these as these are the only members of FFinalSkinVertex that
+					// FSkinnedVerticesDataWrapper reads anyway
+					if (!LastValueEntry.Position.Equals(NewValueEntry.Position) || LastValueEntry.TangentZ != NewValueEntry.TangentZ)
+					{
+						bNewValueIsEqual = false;
+						break;
+					}
+				}
+			}
+		}
+		if (bSkipConstantValues && bNewValueIsEqual)
+		{
+			LastTimeCode = UsdTimeCode;
+			return;
+		}
+
 		pxr::UsdGeomMesh MeshPrim(UsdPrim);
-		ConvertMeshData(FSkinnedVerticesDataWrapper(SkinnedVertices, LODData, bHasVertexColors), MaterialAssignments, UsdTimeCode, UsdPrim, ExportContext, MeshPrim);
+		ConvertMeshData(
+			FSkinnedVerticesDataWrapper(SkinnedVertices, LODData, bHasVertexColors),
+			MaterialAssignments,
+			UsdTimeCode,
+			UsdPrim,
+			ExportContext,
+			MeshPrim
+		);
+
+		if (!bNewValueIsEqual && LastValue.IsSet())
+		{
+			ConvertMeshData(
+				FSkinnedVerticesDataWrapper(LastValue.GetValue(), LODData, bHasVertexColors),
+				MaterialAssignments,
+				LastTimeCode,
+				UsdPrim,
+				ExportContext,
+				MeshPrim
+			);
+		}
+
+		LastValue = SkinnedVertices;
+		LastTimeCode = UsdTimeCode;
 	};
 
 	return true;
diff --git a/Engine/Plugins/Runtime/USDCore/Source/USDUtilities/Private/USDLightConversion.cpp b/Engine/Plugins/Runtime/USDCore/Source/USDUtilities/Private/USDLightConversion.cpp
--- a/Engine/Plugins/Runtime/USDCore/Source/USDUtilities/Private/USDLightConversion.cpp
+++ b/Engine/Plugins/Runtime/USDCore/Source/USDUtilities/Private/USDLightConversion.cpp
@@ -53,37 +53,6 @@ namespace LightConversionImpl
 		return 2.0f * PI * (1.0f - FMath::Cos(FMath::DegreesToRadians(SourceAngleDeg / 2.0f)));
 	}
 
-	// As far as we can tell in Aug 2021, the best approximation for a standard in USD light units is to
-	// use Nits all the time. This doesn't make a lot of sense in some cases but here we try our best to convert to Nits
-	float ConvertIntensityToNits(float Intensity, float Steradians, float AreaInSqMeters, ELightUnits SourceUnits)
-	{
-		switch (SourceUnits)
-		{
-			case ELightUnits::Candelas:
-				// Nit = candela / area
-				return Intensity / AreaInSqMeters;
-				break;
-			case ELightUnits::Lumens:
-				// Nit = lumen / ( sr * area );
-				// https://docs.unrealengine.com/en-US/Engine/Rendering/LightingAndShadows/PhysicalLightUnits/index.html#point,spot,andrectlights
-				return Intensity / (Steradians * AreaInSqMeters);
-				break;
-			case ELightUnits::EV:
-				// Nit = luminance (cd/m2)
-				return EV100ToLuminance(Intensity);
-				break;
-			case ELightUnits::Unitless:
-				// Nit = ( unitless / 625 ) / area = candela / area
-				// https://docs.unrealengine.com/en-US/Engine/Rendering/LightingAndShadows/PhysicalLightUnits/index.html#point,spot,andrectlights
-				return (Intensity / 625.0f) / AreaInSqMeters;
-				break;
-			default:
-				break;
-		}
-
-		return Intensity;
-	}
-
 	// Copied from USpotLightComponent::GetCosHalfConeAngle, so we don't need a component to do the same math
 	float GetSpotLightCosHalfConeAngle(float OuterConeAngle, float InnerConeAngle)
 	{
@@ -601,7 +570,6 @@ bool UnrealToUsd::ConvertRectLightComponent(const URectLightComponent& LightComp
 	// Common for both
 	if (pxr::UsdAttribute Attr = LightAPI.CreateIntensityAttr())
 	{
-		float FinalIntensityNits = 1.0f;
 		float OldIntensity = UsdUtils::GetUsdValue<float>(Attr, UsdTimeCode);
 
 		// Area light with no area probably shouldn't emit any light?
@@ -613,31 +581,8 @@ bool UnrealToUsd::ConvertRectLightComponent(const URectLightComponent& LightComp
 
 		AreaInSqMeters = FMath::Max(AreaInSqMeters, KINDA_SMALL_NUMBER);
 
-		switch (LightComponent.IntensityUnits)
-		{
-			case ELightUnits::Candelas:
-				// Nit = candela / area
-				FinalIntensityNits = OldIntensity / AreaInSqMeters;
-				break;
-			case ELightUnits::Lumens:
-				// Nit = lumen / ( sr * area ); For area lights sr is technically 2PI, but we cancel that with an
-				// extra factor of 2.0 here because URectLightComponent::SetLightBrightness uses just PI and not 2PI as steradian
-				// due to some cosine distribution. This also matches the PI factor between candelas and lumen for rect lights on
-				// https://docs.unrealengine.com/en-US/Engine/Rendering/LightingAndShadows/PhysicalLightUnits/index.html#point,spot,andrectlights
-				FinalIntensityNits = OldIntensity / (PI * AreaInSqMeters);
-				break;
-			case ELightUnits::EV:
-				// Nit = luminance (cd/m2)
-				FinalIntensityNits = EV100ToLuminance(OldIntensity);
-				break;
-			case ELightUnits::Unitless:
-				// Nit = (unitless/625) / area = candela / area
-				// https://docs.unrealengine.com/en-US/Engine/Rendering/LightingAndShadows/PhysicalLightUnits/index.html#point,spot,andrectlights
-				FinalIntensityNits = (OldIntensity / 625.0f) / AreaInSqMeters;
-				break;
-			default:
-				break;
-		}
+		const float Steradians = PI;
+		float FinalIntensityNits = UsdUtils::ConvertIntensityToNits(OldIntensity, Steradians, AreaInSqMeters, LightComponent.IntensityUnits);
 
 		Attr.Set<float>(FinalIntensityNits, UsdTimeCode);
 		UsdUtils::NotifyIfOverriddenOpinion(Attr);
@@ -692,31 +637,8 @@ bool UnrealToUsd::ConvertPointLightComponent(const UPointLightComponent& LightCo
 	const float AreaInSqMeters = FMath::Max(SolidAngle * FMath::Square(LightComponent.SourceRadius / 100.f), KINDA_SMALL_NUMBER);
 	if (pxr::UsdAttribute Attr = Light.CreateIntensityAttr())
 	{
-		float FinalIntensityNits = 1.0f;
 		const float OldIntensity = UsdUtils::GetUsdValue<float>(Attr, UsdTimeCode);
-
-		switch (LightComponent.IntensityUnits)
-		{
-			case ELightUnits::Candelas:
-				// Nit = candela / area
-				FinalIntensityNits = OldIntensity / AreaInSqMeters;
-				break;
-			case ELightUnits::Lumens:
-				// Nit = lumen / (sr * area); For a full sphere sr = 4PI
-				FinalIntensityNits = OldIntensity / (SolidAngle * AreaInSqMeters);
-				break;
-			case ELightUnits::EV:
-				// Nit = luminance (cd/m2)
-				FinalIntensityNits = EV100ToLuminance(OldIntensity);
-				break;
-			case ELightUnits::Unitless:
-				// Nit = (unitless/625) / area = candela / area
-				// https://docs.unrealengine.com/en-US/Engine/Rendering/LightingAndShadows/PhysicalLightUnits/index.html#point,spot,andrectlights
-				FinalIntensityNits = (OldIntensity / 625.0f) / AreaInSqMeters;
-				break;
-			default:
-				break;
-		}
+		float FinalIntensityNits = UsdUtils::ConvertIntensityToNits(OldIntensity, SolidAngle, AreaInSqMeters, LightComponent.IntensityUnits);
 
 		Attr.Set<float>(FinalIntensityNits, UsdTimeCode);
 		UsdUtils::NotifyIfOverriddenOpinion(Attr);
@@ -849,7 +771,7 @@ float UnrealToUsd::ConvertRectLightIntensityProperty(
 	// extra factor of 2.0 here because URectLightComponent::SetLightBrightness uses just PI and not 2PI as steradian
 	// due to some cosine distribution. This also matches the PI factor between candelas and lumen for rect lights on
 	// https://docs.unrealengine.com/en-US/Engine/Rendering/LightingAndShadows/PhysicalLightUnits/index.html#point,spot,andrectlights
-	return LightConversionImpl::ConvertIntensityToNits(UsdIntensity, PI, AreaInSqMeters, SourceUnits);
+	return UsdUtils::ConvertIntensityToNits(UsdIntensity, PI, AreaInSqMeters, SourceUnits);
 }
 
 float UnrealToUsd::ConvertRectLightIntensityProperty(float Intensity, float Radius, const FUsdStageInfo& StageInfo, ELightUnits SourceUnits)
@@ -869,7 +791,7 @@ float UnrealToUsd::ConvertRectLightIntensityProperty(float Intensity, float Radi
 	// extra factor of 2.0 here because URectLightComponent::SetLightBrightness uses just PI and not 2PI as steradian
 	// due to some cosine distribution. This also matches the PI factor between candelas and lumen for rect lights on
 	// https://docs.unrealengine.com/en-US/Engine/Rendering/LightingAndShadows/PhysicalLightUnits/index.html#point,spot,andrectlights
-	return LightConversionImpl::ConvertIntensityToNits(UsdIntensity, PI, AreaInSqMeters, SourceUnits);
+	return UsdUtils::ConvertIntensityToNits(UsdIntensity, PI, AreaInSqMeters, SourceUnits);
 }
 
 float UnrealToUsd::ConvertPointLightIntensityProperty(float Intensity, float SourceRadius, const FUsdStageInfo& StageInfo, ELightUnits SourceUnits)
@@ -883,7 +805,7 @@ float UnrealToUsd::ConvertPointLightIntensityProperty(float Intensity, float Sou
 	// area calculation is possibly incorrect, but I think it depends on the chosen convention
 	const float AreaInSqMeters = FMath::Max(SolidAngle * FMath::Square(SourceRadius / 100.f), KINDA_SMALL_NUMBER);
 
-	return LightConversionImpl::ConvertIntensityToNits(UsdIntensity, SolidAngle, AreaInSqMeters, SourceUnits);
+	return UsdUtils::ConvertIntensityToNits(UsdIntensity, SolidAngle, AreaInSqMeters, SourceUnits);
 }
 
 float UnrealToUsd::ConvertSpotLightIntensityProperty(
@@ -905,7 +827,7 @@ float UnrealToUsd::ConvertSpotLightIntensityProperty(
 	// area calculation is possibly incorrect, but I think it depends on the chosen convention
 	const float AreaInSqMeters = FMath::Max(SolidAngle * FMath::Square(SourceRadius / 100.f), KINDA_SMALL_NUMBER);
 
-	return LightConversionImpl::ConvertIntensityToNits(UsdIntensity, SolidAngle, AreaInSqMeters, SourceUnits);
+	return UsdUtils::ConvertIntensityToNits(UsdIntensity, SolidAngle, AreaInSqMeters, SourceUnits);
 }
 
 float UnrealToUsd::ConvertOuterConeAngleProperty(float OuterConeAngle)
@@ -919,4 +841,33 @@ float UnrealToUsd::ConvertInnerConeAngleProperty(float InnerConeAngle, float Out
 	return FMath::IsNearlyZero(OuterConeAngle) ? 0.0 : 1.0f - InnerConeAngle / OuterConeAngle;
 }
 
+float UsdUtils::ConvertIntensityToNits(float Intensity, float Steradians, float AreaInSqMeters, ELightUnits SourceUnits)
+{
+	switch (SourceUnits)
+	{
+		case ELightUnits::Candelas:
+			// Nit = candela / area
+			return Intensity / AreaInSqMeters;
+			break;
+		case ELightUnits::Lumens:
+			// Nit = lumen / ( sr * area );
+			// https://docs.unrealengine.com/en-US/Engine/Rendering/LightingAndShadows/PhysicalLightUnits/index.html#point,spot,andrectlights
+			return Intensity / (Steradians * AreaInSqMeters);
+			break;
+		case ELightUnits::EV:
+			// Nit = luminance (cd/m2)
+			return EV100ToLuminance(Intensity);
+			break;
+		case ELightUnits::Unitless:
+			// Nit = ( unitless / 625 ) / area = candela / area
+			// https://docs.unrealengine.com/en-US/Engine/Rendering/LightingAndShadows/PhysicalLightUnits/index.html#point,spot,andrectlights
+			return (Intensity / 625.0f) / AreaInSqMeters;
+			break;
+		default:
+			break;
+	}
+
+	return Intensity;
+}
+
 #endif	  // #if USE_USD_SDK
diff --git a/Engine/Plugins/Runtime/USDCore/Source/USDUtilities/Private/USDPrimConversion.cpp b/Engine/Plugins/Runtime/USDCore/Source/USDUtilities/Private/USDPrimConversion.cpp
--- a/Engine/Plugins/Runtime/USDCore/Source/USDUtilities/Private/USDPrimConversion.cpp
+++ b/Engine/Plugins/Runtime/USDCore/Source/USDUtilities/Private/USDPrimConversion.cpp
@@ -105,6 +105,7 @@
 #include "pxr/usd/usdLux/distantLight.h"
 #include "pxr/usd/usdLux/lightAPI.h"
 #include "pxr/usd/usdLux/rectLight.h"
+#include "pxr/usd/usdLux/shadowAPI.h"
 #include "pxr/usd/usdLux/shapingAPI.h"
 #include "pxr/usd/usdLux/sphereLight.h"
 #include "pxr/usd/usdLux/tokens.h"
@@ -132,6 +133,15 @@ static FAutoConsoleVariableRef CVarConsiderAllPrimsHaveAnimatedBounds(
 		 "extremely expensive!")
 );
 
+static bool GSkipConstantValues = true;
+static FAutoConsoleVariableRef CVarSkipConstantValues(
+	TEXT("USD.LevelSequenceExport.SkipConstantValues"),
+	GSkipConstantValues,
+	TEXT(
+		"Whether to prevent the exporter from writing out a timeSample when it has the same value as the previous timeSample. Enable this (default) if you want your generated files to have less timeSamples when possible. Disable this if you want each animated attribute to have an authored value for each timeSample of the animation"
+	)
+);
+
 namespace UE::USDPrimConversion::Private
 {
 	// On the current edit target, will set the Xformable's op order to a single "xformOp:transform",
@@ -4086,6 +4096,200 @@ namespace UE::USDPrimConversion::Private
 			return nullptr;
 		};
 	}
+
+	template<typename UEType, typename UsdType>
+	TFunction<void(UEType, double)> CreateCachedAttrSetter(pxr::UsdAttribute Attr, TFunction<UsdType(UEType)> PostConversion)
+	{
+		if (!Attr)
+		{
+			return [](UEType, double)
+			{
+			};
+		}
+
+		return [LastValue = TOptional<std::decay_t<UEType>>(),
+				LastTimeCode = -DBL_MAX,
+				Attr,
+				PostConversion](UEType NewValue, double NewTimeCode) mutable
+		{
+			bool bNewValueIsEqual = false;
+			if (LastValue.IsSet())
+			{
+				if constexpr (std::is_same_v<std::decay_t<UEType>, FTransform>)
+				{
+					// Transforms don't have an operator== it seems
+					bNewValueIsEqual = LastValue.GetValue().Equals(NewValue);
+				}
+				else if constexpr (std::is_floating_point_v<std::decay_t<UEType>>)
+				{
+					bNewValueIsEqual = FMath::IsNearlyEqual(LastValue.GetValue(), NewValue);
+				}
+				else
+				{
+					bNewValueIsEqual = LastValue.GetValue() == NewValue;
+				}
+			}
+
+			if (GSkipConstantValues && bNewValueIsEqual)
+			{
+				LastTimeCode = NewTimeCode;
+				return;
+			}
+
+			Attr.Set<UsdType>(PostConversion(NewValue), NewTimeCode);
+
+			if (!bNewValueIsEqual && LastValue.IsSet())
+			{
+				// Have to make sure that we write the last timeSample of any stretch where the value stayed constant,
+				// or else we will affect how the resulting curve interpolates from the old value to the new value
+				Attr.Set<UsdType>(PostConversion(LastValue.GetValue()), LastTimeCode);
+			}
+
+			LastValue = NewValue;
+			LastTimeCode = NewTimeCode;
+		};
+	}
+
+	TFunction<void(bool, double)> CreateCachedAttrSetter(pxr::UsdAttribute Attr, pxr::UsdGeomImageable Imageable)
+	{
+		if (!Attr)
+		{
+			return [](bool, double)
+			{
+			};
+		}
+
+		return [LastValue = TOptional<bool>(), LastTimeCode = -DBL_MAX, Attr, Imageable](bool NewValue, double NewTimeCode) mutable
+		{
+			const bool bNewValueIsEqual = LastValue.IsSet() ? LastValue.GetValue() == NewValue : false;
+			if (GSkipConstantValues && bNewValueIsEqual)
+			{
+				LastTimeCode = NewTimeCode;
+				return;
+			}
+
+			if (NewValue)
+			{
+				Imageable.MakeVisible(NewTimeCode);
+			}
+			else
+			{
+				Imageable.MakeInvisible(NewTimeCode);
+			}
+
+			// Imagine our visibility track has a single key that switches to hidden at frame 60.
+			// If our prim is visible by default, MakeVisible will author absolutely nothing, and we'll end up
+			// with a timeSamples that just has '60: "invisible"'. Weirdly enough, in USD that means the prim
+			// will be invisible throughout *the entire duration of the animation* though, which is not what we want.
+			// This check will ensure that if we're visible we should have a value here and not rely on the
+			// fallback value of 'visible', as that doesn't behave how we want.
+			if (!Attr.HasAuthoredValue())
+			{
+				Attr.Set<pxr::TfToken>(NewValue ? pxr::UsdGeomTokens->inherited : pxr::UsdGeomTokens->invisible, NewTimeCode);
+
+				if (!bNewValueIsEqual && LastValue.IsSet())
+				{
+					Attr.Set<pxr::TfToken>(LastValue ? pxr::UsdGeomTokens->inherited : pxr::UsdGeomTokens->invisible, LastTimeCode);
+				}
+			}
+
+			LastValue = NewValue;
+			LastTimeCode = NewTimeCode;
+		};
+	}
+
+	// Overloads intended to handle the translation/rotation/scales of skeletal transforms
+	// (respectively pxr::VtVec3fArray, pxr::VtQuatfArray, pxr::VtVec3hArray)
+	template<typename UEArrayType, typename UsdArrayType>
+	TFunction<void(const UEArrayType&, double)> CreateCachedAttrSetter(
+		pxr::UsdAttribute Attr,
+		TFunction<void(const UEArrayType&, UsdArrayType&)> PostConversion
+	)
+	{
+		if (!Attr)
+		{
+			return [](const UEArrayType&, double)
+			{
+			};
+		}
+
+		return [LastValue = TOptional<UEArrayType>(),
+				LastTimeCode = -DBL_MAX,
+				Attr,
+				PostConversion,
+				ConvertedValues = TUsdStore<UsdArrayType>()	   //
+		](const UEArrayType& NewValue, double NewTimeCode) mutable
+		{
+			bool bNewValueIsEqual = false;
+			if (LastValue.IsSet())
+			{
+				if (LastValue.GetValue().Num() != NewValue.Num())
+				{
+					bNewValueIsEqual = false;
+				}
+				else
+				{
+					bNewValueIsEqual = true;
+					for (int32 Index = 0; Index < NewValue.Num(); ++Index)
+					{
+						// Translations
+						if constexpr (std::is_same_v<UsdArrayType, pxr::VtVec3fArray>)
+						{
+							if (!LastValue.GetValue()[Index].GetTranslation().Equals(NewValue[Index].GetTranslation()))
+							{
+								bNewValueIsEqual = false;
+								break;
+							}
+						}
+						// Rotations
+						else if constexpr (std::is_same_v<UsdArrayType, pxr::VtQuatfArray>)
+						{
+							if (!LastValue.GetValue()[Index].GetRotation().Equals(NewValue[Index].GetRotation()))
+							{
+								bNewValueIsEqual = false;
+								break;
+							}
+						}
+						// Scales
+						else if constexpr (std::is_same_v<UsdArrayType, pxr::VtVec3hArray>)
+						{
+							if (!LastValue.GetValue()[Index].GetScale3D().Equals(NewValue[Index].GetScale3D()))
+							{
+								bNewValueIsEqual = false;
+								break;
+							}
+						}
+						// Blend shape weights
+						else if constexpr (std::is_same_v<UsdArrayType, pxr::VtArray<float>>)
+						{
+							if (!FMath::IsNearlyEqual(LastValue.GetValue()[Index], NewValue[Index]))
+							{
+								bNewValueIsEqual = false;
+								break;
+							}
+						}
+					}
+				}
+			}
+			if (GSkipConstantValues && bNewValueIsEqual)
+			{
+				LastTimeCode = NewTimeCode;
+				return;
+			}
+
+			PostConversion(NewValue, ConvertedValues.Get());
+			Attr.Set(ConvertedValues.Get(), NewTimeCode);
+
+			if (!bNewValueIsEqual && LastValue.IsSet())
+			{
+				PostConversion(LastValue.GetValue(), ConvertedValues.Get());
+				Attr.Set(ConvertedValues.Get(), LastTimeCode);
+			}
+
+			LastValue = NewValue;
+			LastTimeCode = NewTimeCode;
+		};
+	}
 }	 // namespace UE::USDPrimConversion::Private
 
 bool UnrealToUsd::CreateComponentPropertyBaker(
@@ -4106,6 +4310,50 @@ bool UnrealToUsd::CreateComponentPropertyBaker(
 	pxr::UsdStageRefPtr UsdStage = UsdPrim.GetStage();
 	FUsdStageInfo StageInfo{UsdStage};
 
+	TFunction<float(float)> DistanceConversion = [StageInfo](float UEValue)
+	{
+		return UnrealToUsd::ConvertDistance(StageInfo, UEValue);
+	};
+
+	TFunction<float(float)> NoFloatConversion = [StageInfo](float UEValue)
+	{
+		return UEValue;
+	};
+
+	TFunction<bool(bool)> NoBoolConversion = [](bool Value)
+	{
+		return Value;
+	};
+
+	TFunction<pxr::GfVec3f(const FColor&)> ColorConversion = [](const FColor& Value)
+	{
+		pxr::GfVec4f LinearColor = UnrealToUsd::ConvertColor(Value);
+		return pxr::GfVec3f(LinearColor[0], LinearColor[1], LinearColor[2]);
+	};
+
+	TFunction<pxr::GfMatrix4d(const FTransform&)> TransformConversion = [StageInfo](const FTransform& UEValue)
+	{
+		return UnrealToUsd::ConvertTransform(StageInfo, UEValue);
+	};
+
+	TFunction<pxr::VtArray<pxr::GfVec3f>(const FBox&)> BoxConversion = [StageInfo](const FBox& UEBox)
+	{
+		pxr::GfVec3f UEBoundsMinUsdSpace = UnrealToUsd::ConvertVectorFloat(StageInfo, UEBox.Min);
+		pxr::GfVec3f UEBoundsMaxUsdSpace = UnrealToUsd::ConvertVectorFloat(StageInfo, UEBox.Max);
+
+		pxr::GfVec3f UsdMin{
+			FMath::Min(UEBoundsMinUsdSpace[0], UEBoundsMaxUsdSpace[0]),
+			FMath::Min(UEBoundsMinUsdSpace[1], UEBoundsMaxUsdSpace[1]),
+			FMath::Min(UEBoundsMinUsdSpace[2], UEBoundsMaxUsdSpace[2])};
+
+		pxr::GfVec3f UsdMax{
+			FMath::Max(UEBoundsMinUsdSpace[0], UEBoundsMaxUsdSpace[0]),
+			FMath::Max(UEBoundsMinUsdSpace[1], UEBoundsMaxUsdSpace[1]),
+			FMath::Max(UEBoundsMinUsdSpace[2], UEBoundsMaxUsdSpace[2])};
+
+		return pxr::VtArray<pxr::GfVec3f>{UsdMin, UsdMax};
+	};
+
 	// SceneComponent
 	{
 		const static FString TransformPropertyString = UnrealIdentifiers::TransformPropertyName.ToString();
@@ -4177,15 +4425,17 @@ bool UnrealToUsd::CreateComponentPropertyBaker(
 
 			TFunction<const USceneComponent*(void)> ComponentGetter = CreateComponentGetter(&Component);
 			TFunction<USceneComponent*(void)> OriginalAttachParentGetter = CreateComponentGetter(OriginalAttachParent);
+			TFunction<void(const FTransform&, double)> TransformSetter = CreateCachedAttrSetter(Attr, TransformConversion);
 
 			BakerType = EBakingType::Transform;
 			BakerFunction = [CameraCompensation,	//
 							 InverseParentCameraCompensation,
-							 StageInfo,
-							 Attr,
+							 TransformSetter,
 							 ComponentGetter,
-							 OriginalAttachParentGetter](double UsdTimeCode)
+							 OriginalAttachParentGetter](double UsdTimeCode) mutable
 			{
+				TRACE_CPUPROFILER_EVENT_SCOPE(TransformBaker);
+
 				FScopedUsdAllocs Allocs;
 
 				const USceneComponent* Component = ComponentGetter();
@@ -4194,7 +4444,8 @@ bool UnrealToUsd::CreateComponentPropertyBaker(
 					return;
 				}
 
-				// Some setup (like CameraRig_Rail) require that the parent actor on which the component is attached be Tick'ed to update their positions
+				// Some setup (like CameraRig_Rail) require that the parent actor on which the component is attached be Tick'ed to update their
+				// positions
 				if (AActor* AttachParentActor = Component->GetAttachParentActor())
 				{
 					AttachParentActor->Tick(0.0f);
@@ -4222,11 +4473,9 @@ bool UnrealToUsd::CreateComponentPropertyBaker(
 					// could have attach tracks, meaning that we may gain a different attach parent at some point
 					RelativeTransform = Component->GetComponentTransform();
 				}
-
 				RelativeTransform = CameraCompensation * RelativeTransform * InverseParentCameraCompensation;
 
-				pxr::GfMatrix4d UsdTransform = UnrealToUsd::ConvertTransform(StageInfo, RelativeTransform);
-				Attr.Set<pxr::GfMatrix4d>(UsdTransform, UsdTimeCode);
+				TransformSetter(RelativeTransform, UsdTimeCode);
 			};
 		}
 		// bHidden is for the actor, and bHiddenInGame is for a component
@@ -4245,46 +4494,21 @@ bool UnrealToUsd::CreateComponentPropertyBaker(
 			Attr.Clear();
 
 			TFunction<const USceneComponent*(void)> ComponentGetter = CreateComponentGetter(&Component);
+			TFunction<void(bool, double)> VisibilitySetter = CreateCachedAttrSetter(Attr, Imageable);
 
 			BakerType = EBakingType::Visibility;
-			BakerFunction = [Imageable, ComponentGetter](double UsdTimeCode)
+			BakerFunction = [VisibilitySetter, ComponentGetter](double UsdTimeCode)
 			{
+				TRACE_CPUPROFILER_EVENT_SCOPE(VisibilityBaker);
+
 				const USceneComponent* Component = ComponentGetter();
 				if (!Component)
 				{
 					return;
 				}
 
-				if (Component->bHiddenInGame || Component->GetOwner()->IsHidden())
-				{
-					Imageable.MakeInvisible(UsdTimeCode);
-
-					if (pxr::UsdAttribute Attr = Imageable.CreateVisibilityAttr())
-					{
-						if (!Attr.HasAuthoredValue())
-						{
-							Attr.Set<pxr::TfToken>(pxr::UsdGeomTokens->invisible, UsdTimeCode);
-						}
-					}
-				}
-				else
-				{
-					Imageable.MakeVisible(UsdTimeCode);
-
-					// Imagine our visibility track has a single key that switches to hidden at frame 60.
-					// If our prim is visible by default, MakeVisible will author absolutely nothing, and we'll end up
-					// with a timeSamples that just has '60: "invisible"'. Weirdly enough, in USD that means the prim
-					// will be invisible throughout *the entire duration of the animation* though, which is not what we want.
-					// This check will ensure that if we're visible we should have a value here and not rely on the
-					// fallback value of 'visible', as that doesn't behave how we want.
-					if (pxr::UsdAttribute Attr = Imageable.CreateVisibilityAttr())
-					{
-						if (!Attr.HasAuthoredValue())
-						{
-							Attr.Set<pxr::TfToken>(pxr::UsdGeomTokens->inherited, UsdTimeCode);
-						}
-					}
-				}
+				const bool bVisibleInUE = !(Component->bHiddenInGame || Component->GetOwner()->IsHidden());
+				VisibilitySetter(bVisibleInUE, UsdTimeCode);
 			};
 		}
 	}
@@ -4300,23 +4524,86 @@ bool UnrealToUsd::CreateComponentPropertyBaker(
 
 		if (RelevantProperties.Contains(PropertyPath))
 		{
+			pxr::UsdGeomCamera GeomCamera{Prim};
+			if (!GeomCamera)
+			{
+				return false;
+			}
+
 			TFunction<const UCineCameraComponent*(void)> ComponentGetter = CreateComponentGetter(CameraComponent);
 
+			pxr::UsdAttribute FocalLengthAttr = GeomCamera.CreateFocalLengthAttr();
+			pxr::UsdAttribute FocusDistanceAttr = GeomCamera.CreateFocusDistanceAttr();
+			pxr::UsdAttribute FStopAttr = GeomCamera.CreateFStopAttr();
+			pxr::UsdAttribute HorizontalApertureAttr = GeomCamera.CreateHorizontalApertureAttr();
+			pxr::UsdAttribute VerticalApertureAttr = GeomCamera.CreateVerticalApertureAttr();
+
+			TFunction<void(float, double)> FocalLengthSetter = CreateCachedAttrSetter(FocalLengthAttr, DistanceConversion);
+			TFunction<void(float, double)> FocusDistanceSetter = CreateCachedAttrSetter(FocusDistanceAttr, DistanceConversion);
+			TFunction<void(float, double)> FStopSetter = CreateCachedAttrSetter(FStopAttr, NoFloatConversion);
+			TFunction<void(float, double)> HorizontalApertureSetter = CreateCachedAttrSetter(HorizontalApertureAttr, DistanceConversion);
+			TFunction<void(float, double)> VerticalApertureSetter = CreateCachedAttrSetter(VerticalApertureAttr, DistanceConversion);
+
 			BakerType = EBakingType::Camera;
-			BakerFunction = [UsdStage, UsdPrim, ComponentGetter](double UsdTimeCode) mutable
+			BakerFunction = [ComponentGetter,	 //
+							 FocalLengthSetter,
+							 FocusDistanceSetter,
+							 FStopSetter,
+							 HorizontalApertureSetter,
+							 VerticalApertureSetter](double UsdTimeCode) mutable
 			{
+				TRACE_CPUPROFILER_EVENT_SCOPE(CameraBaker);
+
 				const UCineCameraComponent* CameraComponent = ComponentGetter();
 				if (!CameraComponent)
 				{
 					return;
 				}
 
-				UnrealToUsd::ConvertCameraComponent(*CameraComponent, UsdPrim, UsdTimeCode);
+				FocalLengthSetter(CameraComponent->CurrentFocalLength, UsdTimeCode);
+				FocusDistanceSetter(CameraComponent->FocusSettings.ManualFocusDistance, UsdTimeCode);
+				FStopSetter(CameraComponent->CurrentAperture, UsdTimeCode);
+				HorizontalApertureSetter(CameraComponent->Filmback.SensorWidth, UsdTimeCode);
+				VerticalApertureSetter(CameraComponent->Filmback.SensorHeight, UsdTimeCode);
 			};
 		}
 	}
 	else if (const ULightComponentBase* LightComponentBase = Cast<ULightComponentBase>(&Component))
 	{
+		const ULightComponent* LightComponent = Cast<ULightComponent>(LightComponentBase);
+
+		pxr::UsdLuxLightAPI LightAPI(Prim);
+		if (!LightAPI)
+		{
+			return false;
+		}
+
+		pxr::UsdAttribute IntensityAttr = LightAPI.CreateIntensityAttr();
+		pxr::UsdAttribute ColorAttr = LightAPI.CreateColorAttr();
+		pxr::UsdAttribute EnableTemperatureAttr;
+		pxr::UsdAttribute TemperatureAttr;
+		pxr::UsdAttribute ShadowEnableAttr;
+
+		if (LightComponent)
+		{
+			EnableTemperatureAttr = LightAPI.CreateEnableColorTemperatureAttr();
+			TemperatureAttr = LightAPI.CreateColorTemperatureAttr();
+
+			if (!LightComponent->CastShadows)
+			{
+				if (pxr::UsdLuxShadowAPI ShadowAPI = pxr::UsdLuxShadowAPI::Apply(Prim))
+				{
+					ShadowEnableAttr = ShadowAPI.CreateShadowEnableAttr();
+				}
+			}
+		}
+
+		TFunction<void(float, double)> IntensitySetter = CreateCachedAttrSetter(IntensityAttr, NoFloatConversion);
+		TFunction<void(const FColor&, double)> ColorSetter = CreateCachedAttrSetter(ColorAttr, ColorConversion);
+		TFunction<void(bool, double)> EnableTemperatureSetter = CreateCachedAttrSetter(EnableTemperatureAttr, NoBoolConversion);
+		TFunction<void(float, double)> TemperatureSetter = CreateCachedAttrSetter(TemperatureAttr, NoFloatConversion);
+		TFunction<void(bool, double)> ShadowSetter = CreateCachedAttrSetter(ShadowEnableAttr, NoBoolConversion);
+
 		if (const URectLightComponent* RectLightComponent = Cast<URectLightComponent>(LightComponentBase))
 		{
 			static TSet<FString> RelevantProperties =
@@ -4324,19 +4611,63 @@ bool UnrealToUsd::CreateComponentPropertyBaker(
 
 			if (RelevantProperties.Contains(PropertyPath))
 			{
+				pxr::UsdLuxRectLight RectLight{Prim};
+				if (!RectLight)
+				{
+					return false;
+				}
+
+				pxr::UsdAttribute WidthAttr = RectLight.CreateWidthAttr();
+				pxr::UsdAttribute HeightAttr = RectLight.CreateHeightAttr();
+
 				TFunction<const URectLightComponent*(void)> ComponentGetter = CreateComponentGetter(RectLightComponent);
 
+				TFunction<void(float, double)> WidthSetter = CreateCachedAttrSetter(WidthAttr, DistanceConversion);
+				TFunction<void(float, double)> HeightSetter = CreateCachedAttrSetter(WidthAttr, DistanceConversion);
+
 				BakerType = EBakingType::Light;
-				BakerFunction = [UsdStage, UsdPrim, ComponentGetter](double UsdTimeCode) mutable
+				BakerFunction = [ComponentGetter,
+								 IntensitySetter,
+								 ColorSetter,
+								 EnableTemperatureSetter,
+								 TemperatureSetter,
+								 ShadowSetter,
+								 WidthSetter,
+								 HeightSetter](double UsdTimeCode) mutable
 				{
+					TRACE_CPUPROFILER_EVENT_SCOPE(RectLightBaker);
+
 					const URectLightComponent* RectLightComponent = ComponentGetter();
 					if (!RectLightComponent)
 					{
 						return;
 					}
 
-					UnrealToUsd::ConvertLightComponent(*RectLightComponent, UsdPrim, UsdTimeCode);
-					UnrealToUsd::ConvertRectLightComponent(*RectLightComponent, UsdPrim, UsdTimeCode);
+					float UEIntensity = RectLightComponent->Intensity;
+
+					float AreaInSqMeters = (RectLightComponent->SourceWidth / 100.0f) * (RectLightComponent->SourceHeight / 100.0f);
+					if (FMath::IsNearlyZero(AreaInSqMeters))
+					{
+						UEIntensity = 0.0f;
+					}
+
+					const float Steradians = PI;
+					const float FinalIntensityNits = UsdUtils::ConvertIntensityToNits(
+						UEIntensity,
+						Steradians,
+						AreaInSqMeters,
+						RectLightComponent->IntensityUnits
+					);
+
+					IntensitySetter(FinalIntensityNits, UsdTimeCode);
+					WidthSetter(RectLightComponent->SourceWidth, UsdTimeCode);
+					HeightSetter(RectLightComponent->SourceHeight, UsdTimeCode);
+
+					ColorSetter(RectLightComponent->LightColor, UsdTimeCode);
+					EnableTemperatureSetter(RectLightComponent->bUseTemperature, UsdTimeCode);
+					TemperatureSetter(RectLightComponent->Temperature, UsdTimeCode);
+
+					ShadowSetter(RectLightComponent->CastShadows, UsdTimeCode);
 				};
 			}
 		}
@@ -4347,20 +4678,69 @@ bool UnrealToUsd::CreateComponentPropertyBaker(
 
 			if (RelevantProperties.Contains(PropertyPath))
 			{
+				pxr::UsdLuxSphereLight SphereLight{Prim};
+				pxr::UsdLuxShapingAPI ShapingAPI{Prim};
+				if (!SphereLight || !ShapingAPI)
+				{
+					return false;
+				}
+
+				pxr::UsdAttribute RadiusAttr = SphereLight.CreateRadiusAttr();
+				pxr::UsdAttribute TreatAsPointAttr = SphereLight.CreateTreatAsPointAttr();
+				pxr::UsdAttribute AngleAttr = ShapingAPI.CreateShapingConeAngleAttr();
+				pxr::UsdAttribute SoftnessAttr = ShapingAPI.CreateShapingConeSoftnessAttr();
+
 				TFunction<const USpotLightComponent*(void)> ComponentGetter = CreateComponentGetter(SpotLightComponent);
 
+				TFunction<void(bool, double)> TreatAsPointSetter = CreateCachedAttrSetter(TreatAsPointAttr, NoBoolConversion);
+				TFunction<void(float, double)> RadiusSetter = CreateCachedAttrSetter(RadiusAttr, DistanceConversion);
+				TFunction<void(float, double)> AngleSetter = CreateCachedAttrSetter(AngleAttr, NoFloatConversion);
+				TFunction<void(float, double)> SoftnessSetter = CreateCachedAttrSetter(SoftnessAttr, NoFloatConversion);
+
 				BakerType = EBakingType::Light;
-				BakerFunction = [UsdStage, UsdPrim, ComponentGetter](double UsdTimeCode) mutable
+				BakerFunction = [ComponentGetter,
+								 IntensitySetter,
+								 TreatAsPointSetter,
+								 RadiusSetter,
+								 ColorSetter,
+								 EnableTemperatureSetter,
+								 TemperatureSetter,
+								 ShadowSetter,
+								 AngleSetter,
+								 SoftnessSetter](double UsdTimeCode) mutable
 				{
+					TRACE_CPUPROFILER_EVENT_SCOPE(SpotLightBaker);
+
 					const USpotLightComponent* SpotLightComponent = ComponentGetter();
 					if (!SpotLightComponent)
 					{
 						return;
 					}
 
-					UnrealToUsd::ConvertLightComponent(*SpotLightComponent, UsdPrim, UsdTimeCode);
-					UnrealToUsd::ConvertPointLightComponent(*SpotLightComponent, UsdPrim, UsdTimeCode);
-					UnrealToUsd::ConvertSpotLightComponent(*SpotLightComponent, UsdPrim, UsdTimeCode);
+					const float SolidAngle = 2.f * PI * (1.0f - SpotLightComponent->GetCosHalfConeAngle());
+					const float AreaInSqMeters = FMath::Max(SolidAngle * FMath::Square(SpotLightComponent->SourceRadius / 100.f), KINDA_SMALL_NUMBER);
+					const float FinalIntensityNits = UsdUtils::ConvertIntensityToNits(
+						SpotLightComponent->Intensity,
+						SolidAngle,
+						AreaInSqMeters,
+						SpotLightComponent->IntensityUnits
+					);
+
+					IntensitySetter(FinalIntensityNits, UsdTimeCode);
+					TreatAsPointSetter(FMath::IsNearlyZero(SpotLightComponent->SourceRadius), UsdTimeCode);
+					RadiusSetter(SpotLightComponent->SourceRadius, UsdTimeCode);
+					AngleSetter(SpotLightComponent->OuterConeAngle, UsdTimeCode);
+
+					const float Softness = FMath::IsNearlyZero(SpotLightComponent->OuterConeAngle)
+											   ? 0.0
+											   : 1.0f - SpotLightComponent->InnerConeAngle / SpotLightComponent->OuterConeAngle;
+					SoftnessSetter(Softness, UsdTimeCode);
+
+					ColorSetter(SpotLightComponent->LightColor, UsdTimeCode);
+					EnableTemperatureSetter(SpotLightComponent->bUseTemperature, UsdTimeCode);
+					TemperatureSetter(SpotLightComponent->Temperature, UsdTimeCode);
+
+					ShadowSetter(SpotLightComponent->CastShadows, UsdTimeCode);
 				};
 			}
 		}
@@ -4371,19 +4751,59 @@ bool UnrealToUsd::CreateComponentPropertyBaker(
 
 			if (RelevantProperties.Contains(PropertyPath))
 			{
+				pxr::UsdLuxSphereLight SphereLight{Prim};
+				if (!SphereLight)
+				{
+					return false;
+				}
+
+				pxr::UsdAttribute RadiusAttr = SphereLight.CreateRadiusAttr();
+				pxr::UsdAttribute TreatAsPointAttr = SphereLight.CreateTreatAsPointAttr();
+
 				TFunction<const UPointLightComponent*(void)> ComponentGetter = CreateComponentGetter(PointLightComponent);
 
+				TFunction<void(bool, double)> TreatAsPointSetter = CreateCachedAttrSetter(TreatAsPointAttr, NoBoolConversion);
+				TFunction<void(float, double)> RadiusSetter = CreateCachedAttrSetter(RadiusAttr, DistanceConversion);
+
 				BakerType = EBakingType::Light;
-				BakerFunction = [UsdStage, UsdPrim, ComponentGetter](double UsdTimeCode) mutable
+				BakerFunction = [ComponentGetter,
+								 IntensitySetter,
+								 TreatAsPointSetter,
+								 RadiusSetter,
+								 ColorSetter,
+								 EnableTemperatureSetter,
+								 TemperatureSetter,
+								 ShadowSetter](double UsdTimeCode) mutable
 				{
+					TRACE_CPUPROFILER_EVENT_SCOPE(PointLightBaker);
+
 					const UPointLightComponent* PointLightComponent = ComponentGetter();
 					if (!PointLightComponent)
 					{
 						return;
 					}
 
-					UnrealToUsd::ConvertLightComponent(*PointLightComponent, UsdPrim, UsdTimeCode);
-					UnrealToUsd::ConvertPointLightComponent(*PointLightComponent, UsdPrim, UsdTimeCode);
+					const float SolidAngle = 4.f * PI;
+					const float AreaInSqMeters = FMath::Max(
+						SolidAngle * FMath::Square(PointLightComponent->SourceRadius / 100.f),
+						KINDA_SMALL_NUMBER
+					);
+					const float FinalIntensityNits = UsdUtils::ConvertIntensityToNits(
+						PointLightComponent->Intensity,
+						SolidAngle,
+						AreaInSqMeters,
+						PointLightComponent->IntensityUnits
+					);
+
+					IntensitySetter(FinalIntensityNits, UsdTimeCode);
+					TreatAsPointSetter(FMath::IsNearlyZero(PointLightComponent->SourceRadius), UsdTimeCode);
+					RadiusSetter(PointLightComponent->SourceRadius, UsdTimeCode);
+
+					ColorSetter(PointLightComponent->LightColor, UsdTimeCode);
+					EnableTemperatureSetter(PointLightComponent->bUseTemperature, UsdTimeCode);
+					TemperatureSetter(PointLightComponent->Temperature, UsdTimeCode);
+
+					ShadowSetter(PointLightComponent->CastShadows, UsdTimeCode);
 				};
 			}
 		}
@@ -4394,35 +4814,81 @@ bool UnrealToUsd::CreateComponentPropertyBaker(
 
 			if (RelevantProperties.Contains(PropertyPath))
 			{
+				pxr::UsdLuxDistantLight Light{Prim};
+				if (!Light)
+				{
+					return false;
+				}
+
+				pxr::UsdAttribute AngleAttr = Light.CreateAngleAttr();
+
 				TFunction<const UDirectionalLightComponent*(void)> ComponentGetter = CreateComponentGetter(DirectionalLightComponent);
 
+				TFunction<void(float, double)> AngleSetter = CreateCachedAttrSetter(AngleAttr, NoFloatConversion);
+
 				BakerType = EBakingType::Light;
-				BakerFunction = [UsdStage, UsdPrim, ComponentGetter](double UsdTimeCode) mutable
+				BakerFunction = [ComponentGetter,	 //
+								 IntensitySetter,
+								 AngleSetter,
+								 ColorSetter,
+								 EnableTemperatureSetter,
+								 TemperatureSetter,
+								 ShadowSetter](double UsdTimeCode) mutable
 				{
+					TRACE_CPUPROFILER_EVENT_SCOPE(DirectionalLightBaker);
+
 					const UDirectionalLightComponent* DirectionalLightComponent = ComponentGetter();
 					if (!DirectionalLightComponent)
 					{
 						return;
 					}
 
-					UnrealToUsd::ConvertLightComponent(*DirectionalLightComponent, UsdPrim, UsdTimeCode);
-					UnrealToUsd::ConvertDirectionalLightComponent(*DirectionalLightComponent, UsdPrim, UsdTimeCode);
+					IntensitySetter(DirectionalLightComponent->Intensity, UsdTimeCode);
+					AngleSetter(DirectionalLightComponent->LightSourceAngle, UsdTimeCode);
+
+					ColorSetter(DirectionalLightComponent->LightColor, UsdTimeCode);
+					EnableTemperatureSetter(DirectionalLightComponent->bUseTemperature, UsdTimeCode);
+					TemperatureSetter(DirectionalLightComponent->Temperature, UsdTimeCode);
+
+					ShadowSetter(DirectionalLightComponent->CastShadows, UsdTimeCode);
 				};
 			}
 		}
 	}
 	else if (const UUsdDrawModeComponent* DrawModeComponent = Cast<UUsdDrawModeComponent>(&Component))
 	{
+		// We don't support importing/exporting animated texture cards for now, and the other UsdGeomModelAPI
+		// attributes are uniform. The only animation we can export from this are the extents
 		static TSet<FString> RelevantProperties = {
 			GET_MEMBER_NAME_CHECKED(UUsdDrawModeComponent, BoundsMin).ToString(),
 			GET_MEMBER_NAME_CHECKED(UUsdDrawModeComponent, BoundsMax).ToString()};
 
 		if (RelevantProperties.Contains(PropertyPath))
 		{
+			pxr::UsdAttribute ExtentsAttr;
+			if (pxr::UsdGeomBoundable Boundable{UsdPrim})
+			{
+				// Try using the extents attribute if we already have one authored
+				ExtentsAttr = Boundable.GetExtentAttr();
+			}
+			if (!ExtentsAttr)
+			{
+				if (pxr::UsdGeomModelAPI GeomModelAPI = pxr::UsdGeomModelAPI::Apply(UsdPrim))
+				{
+					ExtentsAttr = GeomModelAPI.GetExtentsHintAttr();
+				}
+			}
+			if (!ExtentsAttr)
+			{
+				return false;
+			}
+
 			TFunction<const UUsdDrawModeComponent*(void)> ComponentGetter = CreateComponentGetter(DrawModeComponent);
 
+			TFunction<void(const FBox&, double)> ExtentsSetter = CreateCachedAttrSetter(ExtentsAttr, BoxConversion);
+
 			BakerType = EBakingType::Bounds;
-			BakerFunction = [UsdStage, UsdPrim, ComponentGetter](double UsdTimeCode) mutable
+			BakerFunction = [ComponentGetter, ExtentsSetter](double UsdTimeCode) mutable
 			{
 				const UUsdDrawModeComponent* DrawModeComponent = ComponentGetter();
 				if (!DrawModeComponent)
@@ -4430,8 +4896,8 @@ bool UnrealToUsd::CreateComponentPropertyBaker(
 					return;
 				}
 
-				const bool bWriteExtents = true;
-				UnrealToUsd::ConvertDrawModeComponent(*DrawModeComponent, UsdPrim, bWriteExtents, UsdTimeCode);
+				const FBox NewBox{DrawModeComponent->BoundsMin, DrawModeComponent->BoundsMax};
+				ExtentsSetter(NewBox, UsdTimeCode);
 			};
 		}
 	}
@@ -4546,96 +5012,165 @@ bool UnrealToUsd::CreateSkeletalAnimationBaker(
 		BlendShapesAttr.Set(BlendShapeNames);
 	}
 
-	pxr::VtVec3fArray Translations;
-	pxr::VtQuatfArray Rotations;
-	pxr::VtVec3hArray Scales;
-	pxr::VtArray<float> BlendShapeWeights;
-
-	OutBaker.ComponentPath = Component.GetPathName();
-	OutBaker.BakerType = EBakingType::Skeletal;
-	OutBaker.BakerFunction = [&Component,
-							  StageInfo,
-							  Translations,
-							  Rotations,
-							  Scales,
-							  BlendShapeWeights,
-							  TranslationsAttr,
-							  RotationsAttr,
-							  ScalesAttr,
-							  BlendShapeWeightsAttr,
-							  NumBones,
-							  NumMorphTargets](double UsdTimeCode) mutable
-	{
-		FScopedUsdAllocs InnerAllocs;
+	TFunction<USkeletalMeshComponent*(void)> ComponentGetter = CreateComponentGetter(&Component);
 
+	TFunction<void(const TArray<FTransform>&, pxr::VtVec3fArray&)> BoneTranslationConversion =
+		[StageInfo](const TArray<FTransform>& UEBones, pxr::VtVec3fArray& Translations)
+	{
+		const int32 NumBones = UEBones.Num();
 		Translations.resize(NumBones);
-		Rotations.resize(NumBones);
-		Scales.resize(NumBones);
 
-		if (USkeletalMeshComponent* Leader = Cast<USkeletalMeshComponent>(Component.LeaderPoseComponent.Get()))
+		for (int32 BoneIndex = 0; BoneIndex < NumBones; ++BoneIndex)
 		{
-			UsdUtils::RefreshSkeletalMeshComponent(*Leader);
+			// Pulled from UsdUtils::ConvertAxes since we don't care about the other components of the transform here
+			FVector Translation = UEBones[BoneIndex].GetTranslation();
+			if (StageInfo.UpAxis == EUsdUpAxis::ZAxis)
+			{
+				Translation.Y = -Translation.Y;
+			}
+			else
+			{
+				Swap(Translation.Y, Translation.Z);
+			}
+
+			Translations[BoneIndex] = UnrealToUsd::ConvertVectorFloat(Translation) * (0.01f / StageInfo.MetersPerUnit);
 		}
-		UsdUtils::RefreshSkeletalMeshComponent(Component);
+	};
 
-		// I'm not entirely sure why this is needed but FFbxExporter::ExportAnimTrack and FFbxExporter::ExportLevelSequenceBaked3DTransformTrack
-		// do this so for safety maybe we should as well?
-		if (AActor* Owner = Component.GetOwner())
+	TFunction<void(const TArray<FTransform>&, pxr::VtQuatfArray&)> BoneRotationConversion =
+		[StageInfo](const TArray<FTransform>& UEBones, pxr::VtQuatfArray& Rotations)
+	{
+		const int32 NumBones = UEBones.Num();
+
+		Rotations.resize(NumBones);
+
+		for (int32 BoneIndex = 0; BoneIndex < NumBones; ++BoneIndex)
 		{
-			Owner->Tick(0.0f);
+			// Pulled from UsdUtils::ConvertAxes since we don't care about the other components of the transform here
+			FQuat Rotation = UEBones[BoneIndex].GetRotation();
+			if (StageInfo.UpAxis == EUsdUpAxis::ZAxis)
+			{
+				Rotation.X = -Rotation.X;
+				Rotation.Z = -Rotation.Z;
+			}
+			else
+			{
+				Rotation = Rotation.Inverse();
+				Swap(Rotation.Y, Rotation.Z);
+			}
+
+			Rotations[BoneIndex] = UnrealToUsd::ConvertQuatFloat(Rotation).GetNormalized();
 		}
+	};
 
-		TArray<FTransform> LocalBoneTransforms;
-		UsdUtils::GetBoneTransforms(&Component, LocalBoneTransforms);
+	TFunction<void(const TArray<FTransform>&, pxr::VtVec3hArray&)> BoneScaleConversion =
+		[StageInfo](const TArray<FTransform>& UEBones, pxr::VtVec3hArray& Scales)
+	{
+		const int32 NumBones = UEBones.Num();
+
+		Scales.resize(NumBones);
 
-		// For whatever reason it seems that sometimes this is not ready for us, so let's force it to be recalculated
-		if (LocalBoneTransforms.Num() == 0)
+		for (int32 BoneIndex = 0; BoneIndex < NumBones; ++BoneIndex)
 		{
-			const int32 LODIndex = 0;
-			Component.RecalcRequiredBones(LODIndex);
+			// Pulled from UsdUtils::ConvertAxes since we don't care about the other components of the transform here
+			FVector Scale = UEBones[BoneIndex].GetScale3D();
+			if (StageInfo.UpAxis != EUsdUpAxis::ZAxis)
+			{
+				Swap(Scale.Y, Scale.Z);
+			}
+
+			Scales[BoneIndex] = UnrealToUsd::ConvertVectorHalf(Scale);
 		}
-		if (LocalBoneTransforms.Num() != NumBones)
+	};
+
+	TFunction<void(const TArray<float>&, pxr::VtArray<float>&)> WeightConversion = [](const TArray<float>& UEWeights, pxr::VtArray<float>& UsdWeights)
+	{
+		UsdWeights.resize(UEWeights.Num());
+		FMemory::Memcpy(UsdWeights.data(), UEWeights.GetData(), UEWeights.GetTypeSize() * UEWeights.Num());
+	};
+
+	TFunction<void(const TArray<FTransform>&, double)> BoneTranslationSetter = CreateCachedAttrSetter(	  //
+		TranslationsAttr,
+		BoneTranslationConversion
+	);
+
+	TFunction<void(const TArray<FTransform>&, double)> BoneRotationSetter = CreateCachedAttrSetter(	   //
+		RotationsAttr,
+		BoneRotationConversion
+	);
+
+	TFunction<void(const TArray<FTransform>&, double)> BoneScaleSetter = CreateCachedAttrSetter(	//
+		ScalesAttr,
+		BoneScaleConversion
+	);
+
+	TFunction<void(const TArray<float>&, double)> WeightSetter = CreateCachedAttrSetter(	//
+		BlendShapeWeightsAttr,
+		WeightConversion
+	);
+
+	OutBaker.ComponentPath = Component.GetPathName();
+	OutBaker.BakerType = EBakingType::Skeletal;
+	OutBaker.BakerFunction =
+		[ComponentGetter, BoneTranslationSetter, BoneRotationSetter, BoneScaleSetter, WeightSetter, NumBones](double UsdTimeCode) mutable
+	{
+		TRACE_CPUPROFILER_EVENT_SCOPE(SkeletalBaker);
+
+		USkeletalMeshComponent* Component = ComponentGetter();
+		if (!Component)
 		{
-			UE_LOG(
-				LogUsd,
-				Warning,
-				TEXT("Failed to retrieve bone transforms when baking skeletal animation for component '%s' at timeCode '%f'. Expected %d transforms, "
-					 "received %d"),
-				*Component.GetPathName(),
-				UsdTimeCode,
-				NumBones,
-				LocalBoneTransforms.Num()
-			);
 			return;
 		}
 
-		for (int32 BoneIndex = 0; BoneIndex < NumBones; ++BoneIndex)
+		if (USkeletalMeshComponent* Leader = Cast<USkeletalMeshComponent>(Component->LeaderPoseComponent.Get()))
 		{
-			FTransform BoneTransform = LocalBoneTransforms[BoneIndex];
-			BoneTransform = UsdUtils::ConvertAxes(StageInfo.UpAxis == EUsdUpAxis::ZAxis, BoneTransform);
+			UsdUtils::RefreshSkeletalMeshComponent(*Leader);
+		}
+		UsdUtils::RefreshSkeletalMeshComponent(*Component);
 
-			Translations[BoneIndex] = UnrealToUsd::ConvertVectorFloat(BoneTransform.GetTranslation()) * (0.01f / StageInfo.MetersPerUnit);
-			Rotations[BoneIndex] = UnrealToUsd::ConvertQuatFloat(BoneTransform.GetRotation()).GetNormalized();
-			Scales[BoneIndex] = UnrealToUsd::ConvertVectorHalf(BoneTransform.GetScale3D());
+		// I'm not entirely sure why this is needed but FFbxExporter::ExportAnimTrack and FFbxExporter::ExportLevelSequenceBaked3DTransformTrack
+		// do this so for safety maybe we should as well?
+		if (AActor* Owner = Component->GetOwner())
+		{
+			Owner->Tick(0.0f);
 		}
 
-		if (Translations.size() > 0)
+		// Handle morph target weights
 		{
-			TranslationsAttr.Set(Translations, UsdTimeCode);
-			RotationsAttr.Set(Rotations, UsdTimeCode);
-			ScalesAttr.Set(Scales, UsdTimeCode);
+			WeightSetter(Component->MorphTargetWeights, UsdTimeCode);
 		}
 
-		if (NumMorphTargets > 0 && BlendShapeWeightsAttr)
+		// Handle bone transforms
 		{
-			BlendShapeWeights.resize(NumMorphTargets);
+			TArray<FTransform> LocalBoneTransforms;
+			UsdUtils::GetBoneTransforms(Component, LocalBoneTransforms);
 
-			for (int32 MorphTargetIndex = 0; MorphTargetIndex < NumMorphTargets; ++MorphTargetIndex)
+			// For whatever reason it seems that sometimes this is not ready for us, so let's force it to be recalculated
+			if (LocalBoneTransforms.Num() == 0)
 			{
-				BlendShapeWeights[MorphTargetIndex] = Component.MorphTargetWeights[MorphTargetIndex];
+				const int32 LODIndex = 0;
+				Component->RecalcRequiredBones(LODIndex);
+			}
+			if (LocalBoneTransforms.Num() != NumBones)
+			{
+				UE_LOG(
+					LogUsd,
+					Warning,
+					TEXT(
+						"Failed to retrieve bone transforms when baking skeletal animation for component '%s' at timeCode '%f'. Expected %d transforms, "
+						"received %d"
+					),
+					*Component->GetPathName(),
+					UsdTimeCode,
+					NumBones,
+					LocalBoneTransforms.Num()
+				);
+				return;
 			}
 
-			BlendShapeWeightsAttr.Set(BlendShapeWeights, UsdTimeCode);
+			BoneTranslationSetter(LocalBoneTransforms, UsdTimeCode);
+			BoneRotationSetter(LocalBoneTransforms, UsdTimeCode);
+			BoneScaleSetter(LocalBoneTransforms, UsdTimeCode);
 		}
 	};
 
@@ -5328,10 +5863,10 @@ bool UnrealToUsd::ConvertXformable(
 
 	FMovieSceneInverseSequenceTransform SequenceToRootTransform = SequenceTransform.Inverse();
 
-	auto EvaluateChannel = [&PlaybackRange,
-							&Resolution,
-							&DisplayRate,
-							&SequenceToRootTransform](const FMovieSceneDoubleChannel* Channel, double DefaultValue) -> TArray<TPair<FFrameNumber, float>>
+	auto EvaluateChannel = [&PlaybackRange, &Resolution, &DisplayRate, &SequenceToRootTransform](
+							   const FMovieSceneDoubleChannel* Channel,
+							   double DefaultValue
+						   ) -> TArray<TPair<FFrameNumber, float>>
 	{
 		TArray<TPair<FFrameNumber, float>> Values;
 
diff --git a/Engine/Plugins/Runtime/USDCore/Source/USDUtilities/Public/USDLightConversion.h b/Engine/Plugins/Runtime/USDCore/Source/USDUtilities/Public/USDLightConversion.h
--- a/Engine/Plugins/Runtime/USDCore/Source/USDUtilities/Public/USDLightConversion.h
+++ b/Engine/Plugins/Runtime/USDCore/Source/USDUtilities/Public/USDLightConversion.h
@@ -188,4 +188,11 @@ namespace UnrealToUsd
 	USDUTILITIES_API float ConvertInnerConeAngleProperty(float InnerConeAngle, float OuterConeAngle);	 // Returns ConeSoftnessAttr() in USD
 }
 
+namespace UsdUtils
+{
+	// As far as we can tell in Aug 2021, the best approximation for a standard in USD light units is to
+	// use Nits all the time. This doesn't make a lot of sense in some cases but here we try our best to convert to Nits
+	float ConvertIntensityToNits(float Intensity, float Steradians, float AreaInSqMeters, ELightUnits SourceUnits);
+}
+
 #endif	  // #if USE_USD_SDK
