USD:
- Translate MaterialX files referenced as payloads

diff --git a/Engine/Plugins/Runtime/USDCore/Source/USDUtilities/Private/USDShadeConversion.cpp b/Engine/Plugins/Runtime/USDCore/Source/USDUtilities/Private/USDShadeConversion.cpp
--- a/Engine/Plugins/Runtime/USDCore/Source/USDUtilities/Private/USDShadeConversion.cpp
+++ b/Engine/Plugins/Runtime/USDCore/Source/USDUtilities/Private/USDShadeConversion.cpp
@@ -3044,6 +3044,23 @@ TArray<FString> UsdUtils::GetMaterialXFilePaths(const pxr::UsdPrim& Prim)
 
 	TArray<FString> Result;
 
+	auto AddMaterialXFilePath = [&Result](const pxr::UsdPrimCompositionQueryArc& CompositionArc, auto ReferenceOrPayloadEditor, auto UsdReferenceOrPayload)
+	{
+		if(CompositionArc.GetIntroducingListEditor(&ReferenceOrPayloadEditor, &UsdReferenceOrPayload))
+		{
+			// The mtlx file will be the "target layer". It's useful to get its real path via the SdfLayer
+			// interface as it already makes sure it is absolute
+			pxr::SdfLayerHandle TargetLayer = CompositionArc.GetTargetLayer();
+			FString AbsoluteFilePath = UsdToUnreal::ConvertString(TargetLayer->GetRealPath());
+
+			FString Extension = FPaths::GetExtension(AbsoluteFilePath);
+			if(Extension == UnrealIdentifiers::MaterialXRenderContext && FPaths::FileExists(AbsoluteFilePath))
+			{
+				Result.Add(AbsoluteFilePath);
+			}
+		}
+	};
+
 	// We used to just fetch "direct references" here, but stages may compose the .mtlx file reference onto the prim
 	// via another sublayer or reference, so it will be marked as an "ancestral arc" and not be included in the "direct references" filter
 	pxr::UsdPrimCompositionQuery PrimCompositionQuery = pxr::UsdPrimCompositionQuery{ Prim };
@@ -3051,22 +3068,11 @@ TArray<FString> UsdUtils::GetMaterialXFilePaths(const pxr::UsdPrim& Prim)
 	{
 		if(CompositionArc.GetArcType() == pxr::PcpArcTypeReference)
 		{
-			pxr::SdfReferenceEditorProxy ReferenceEditor;
-			pxr::SdfReference UsdReference;
-
-			if(CompositionArc.GetIntroducingListEditor(&ReferenceEditor, &UsdReference))
-			{
-				// The mtlx file will be the "target layer". It's useful to get its real path via the SdfLayer
-				// interface as it already makes sure it is absolute
-				pxr::SdfLayerHandle TargetLayer = CompositionArc.GetTargetLayer();
-				FString AbsoluteFilePath = UsdToUnreal::ConvertString(TargetLayer->GetRealPath());
-
-				FString Extension = FPaths::GetExtension(AbsoluteFilePath);
-				if(Extension == UnrealIdentifiers::MaterialXRenderContext && FPaths::FileExists(AbsoluteFilePath))
-				{
-					Result.Add(AbsoluteFilePath);
-				}
-			}
+			AddMaterialXFilePath(CompositionArc, pxr::SdfReferenceEditorProxy{}, pxr::SdfReference{});
+		}
+		else if(CompositionArc.GetArcType() == pxr::PcpArcType::PcpArcTypePayload)
+		{
+			AddMaterialXFilePath(CompositionArc, pxr::SdfPayloadEditorProxy{}, pxr::SdfPayload{});
 		}
 	}
 
