USD: Suppress asset export analytics events during the level export;

Also export the count of geometry caches and cubemaps when exporting levels to USD.

diff --git a/Engine/Plugins/Importers/USDImporter/Source/USDExporter/Private/USDConversionBlueprintLibrary.cpp b/Engine/Plugins/Importers/USDImporter/Source/USDExporter/Private/USDConversionBlueprintLibrary.cpp
--- a/Engine/Plugins/Importers/USDImporter/Source/USDExporter/Private/USDConversionBlueprintLibrary.cpp
+++ b/Engine/Plugins/Importers/USDImporter/Source/USDExporter/Private/USDConversionBlueprintLibrary.cpp
@@ -633,6 +633,16 @@ void UUsdConversionBlueprintLibrary::SendAnalytics(
 	IUsdClassesModule::SendAnalytics(MoveTemp(Converted), EventName, bAutomated, ElapsedSeconds, NumberOfFrames, Extension);
 }
 
+void UUsdConversionBlueprintLibrary::BlockAnalyticsEvents()
+{
+	IUsdClassesModule::BlockAnalyticsEvents();
+}
+
+void UUsdConversionBlueprintLibrary::ResumeAnalyticsEvents()
+{
+	IUsdClassesModule::ResumeAnalyticsEvents();
+}
+
 void UUsdConversionBlueprintLibrary::RemoveAllPrimSpecs(const FString& StageRootLayer, const FString& PrimPath, const FString& TargetLayer)
 {
 #if USE_USD_SDK
diff --git a/Engine/Plugins/Importers/USDImporter/Source/USDExporter/Public/USDConversionBlueprintLibrary.h b/Engine/Plugins/Importers/USDImporter/Source/USDExporter/Public/USDConversionBlueprintLibrary.h
--- a/Engine/Plugins/Importers/USDImporter/Source/USDExporter/Public/USDConversionBlueprintLibrary.h
+++ b/Engine/Plugins/Importers/USDImporter/Source/USDExporter/Public/USDConversionBlueprintLibrary.h
@@ -148,6 +148,12 @@ public:
 		const FString& Extension
 	);
 
+	UFUNCTION(BlueprintCallable, Category = "USD|Analytics")
+	static void BlockAnalyticsEvents();
+
+	UFUNCTION(BlueprintCallable, Category = "USD|Analytics")
+	static void ResumeAnalyticsEvents();
+
 	/**
 	 * Removes all the prim specs for Prim on the given Layer.
 	 *
diff --git a/Engine/Plugins/Runtime/USDCore/Source/USDClasses/Private/USDClassesModule.cpp b/Engine/Plugins/Runtime/USDCore/Source/USDClasses/Private/USDClassesModule.cpp
--- a/Engine/Plugins/Runtime/USDCore/Source/USDClasses/Private/USDClassesModule.cpp
+++ b/Engine/Plugins/Runtime/USDCore/Source/USDClasses/Private/USDClassesModule.cpp
@@ -78,6 +78,7 @@ namespace UE::USDClasses::Private
 	}
 
 	TMap<FString, int32> PackagePathNameToDirtyCounter;
+	int AnalyticsAllowedWhenZero = 0;
 }
 
 void IUsdClassesModule::UpdatePlugInfoFiles(const FString& PluginDirectory, const FString& TargetDllFolder)
@@ -236,7 +237,7 @@ void IUsdClassesModule::SendAnalytics(
 	const FString& Extension
 )
 {
-	if (FEngineAnalytics::IsAvailable())
+	if (FEngineAnalytics::IsAvailable() && UE::USDClasses::Private::AnalyticsAllowedWhenZero == 0)
 	{
 		TArray<FAnalyticsEventAttribute> Attributes(InAttributes);
 
@@ -255,7 +256,7 @@ void IUsdClassesModule::SendAnalytics(
 
 void IUsdClassesModule::SendAnalytics(TArray<FAnalyticsEventAttribute>&& InAttributes, const FString& EventName)
 {
-	if (FEngineAnalytics::IsAvailable())
+	if (FEngineAnalytics::IsAvailable() && UE::USDClasses::Private::AnalyticsAllowedWhenZero == 0)
 	{
 		TArray<FAnalyticsEventAttribute> Attributes(InAttributes);
 
@@ -268,6 +269,16 @@ void IUsdClassesModule::SendAnalytics(TArray<FAnalyticsEventAttribute>&& InAttri
 	}
 }
 
+void IUsdClassesModule::BlockAnalyticsEvents()
+{
+	UE::USDClasses::Private::AnalyticsAllowedWhenZero += 1;
+}
+
+void IUsdClassesModule::ResumeAnalyticsEvents()
+{
+	UE::USDClasses::Private::AnalyticsAllowedWhenZero -= 1;
+}
+
 bool IUsdClassesModule::HashObjectPackage(const UObject* Object, FSHA1& HashToUpdate)
 {
 #if WITH_EDITOR
diff --git a/Engine/Plugins/Runtime/USDCore/Source/USDClasses/Public/USDClassesModule.h b/Engine/Plugins/Runtime/USDCore/Source/USDClasses/Public/USDClassesModule.h
--- a/Engine/Plugins/Runtime/USDCore/Source/USDClasses/Public/USDClassesModule.h
+++ b/Engine/Plugins/Runtime/USDCore/Source/USDClasses/Public/USDClassesModule.h
@@ -37,6 +37,10 @@ public:
 	);
 	USDCLASSES_API static void SendAnalytics(TArray<FAnalyticsEventAttribute>&& InAttributes, const FString& EventName);
 
+	USDCLASSES_API static void BlockAnalyticsEvents();
+
+	USDCLASSES_API static void ResumeAnalyticsEvents();
+
 	/**
 	 * Updates HashToUpdate with the Object's package's persistent guid, the corresponding file save
 	 * date and time, and the number of times the package has been dirtied since last being saved.
