USD: Utils for logging notice information.

diff --git a/Engine/Plugins/Runtime/USDCore/Source/UnrealUSDWrapper/Private/USDListener.cpp b/Engine/Plugins/Runtime/USDCore/Source/UnrealUSDWrapper/Private/USDListener.cpp
--- a/Engine/Plugins/Runtime/USDCore/Source/UnrealUSDWrapper/Private/USDListener.cpp
+++ b/Engine/Plugins/Runtime/USDCore/Source/UnrealUSDWrapper/Private/USDListener.cpp
@@ -23,12 +23,333 @@
 #include "pxr/usd/usd/stage.h"
 #include "pxr/usd/usdGeom/tokens.h"
 #include "USDIncludesEnd.h"
-
 #endif	  // USE_USD_SDK
 
-namespace UsdToUnreal
+#define ENABLE_NOTICE_LOGGING 0
+
+namespace UE::USDListener::Private
 {
 #if USE_USD_SDK
+
+#if ENABLE_NOTICE_LOGGING
+	void LogChangeListEntry(const pxr::SdfChangeList::Entry& Entry, int IndentLevel)
+	{
+		FScopedUsdAllocs Allocs;
+
+		FString Indent;
+		for (int Index = 0; Index < IndentLevel; ++Index)
+		{
+			Indent += TEXT("\t");
+		}
+
+		UE_LOG(LogUsd, Log, TEXT("%sChangeListEntry:"), *Indent);
+
+		UE_LOG(LogUsd, Log, TEXT("%s\tInfoChanges:"), *Indent);
+		for (const std::pair<pxr::TfToken, std::pair<pxr::VtValue, pxr::VtValue>>& AttributeChange : Entry.infoChanged)
+		{
+			const FString FieldToken = UTF8_TO_TCHAR(AttributeChange.first.GetString().c_str());
+
+			const pxr::VtValue& OldValue = AttributeChange.second.first;
+			const pxr::VtValue& NewValue = AttributeChange.second.second;
+			std::string OldValueString = pxr::TfStringify(OldValue);
+			std::string NewValueString = pxr::TfStringify(NewValue);
+
+			UE_LOG(
+				LogUsd,
+				Log,
+				TEXT("%s\t\t'%s': From '%s' to '%s'"),
+				*Indent,
+				*FieldToken,
+				UTF8_TO_TCHAR(OldValueString.c_str()),
+				UTF8_TO_TCHAR(NewValueString.c_str())
+			);
+		}
+
+		UE_LOG(LogUsd, Log, TEXT("%s\tSubLayerChanges:"), *Indent);
+		static_assert(static_cast<int>(pxr::SdfChangeList::SubLayerChangeType::SubLayerAdded) == 0, "Enum values changed!");
+		static_assert(static_cast<int>(pxr::SdfChangeList::SubLayerChangeType::SubLayerRemoved) == 1, "Enum values changed!");
+		static_assert(static_cast<int>(pxr::SdfChangeList::SubLayerChangeType::SubLayerOffset) == 2, "Enum values changed!");
+		const TCHAR* SubLayerChangeTypeStr[] = {
+			TEXT("SubLayerAdded"),
+			TEXT("SubLayerRemoved"),
+			TEXT("SubLayerOffset"),
+		};
+		for (const std::pair<std::string, pxr::SdfChangeList::SubLayerChangeType>& SubLayerChange : Entry.subLayerChanges)
+		{
+			UE_LOG(
+				LogUsd,
+				Log,
+				TEXT("%s\t\t'%s': change type '%s'"),
+				*Indent,
+				UTF8_TO_TCHAR(SubLayerChange.first.c_str()),
+				SubLayerChangeTypeStr[static_cast<int>(SubLayerChange.second)]
+			);
+		}
+
+		UE_LOG(LogUsd, Log, TEXT("%s\tOldPath: '%s'"), *Indent, UTF8_TO_TCHAR(Entry.oldPath.GetString().c_str()));
+
+		UE_LOG(LogUsd, Log, TEXT("%s\tOldIdentifier: '%s'"), *Indent, UTF8_TO_TCHAR(Entry.oldIdentifier.c_str()));
+
+		UE_LOG(LogUsd, Log, TEXT("%s\tFlags:"), *Indent);
+		TArray<FString> FlagsToLog;
+#define TRY_LOG_FLAG(FlagName)           \
+	if (Entry.flags.##FlagName)          \
+	{                                    \
+		FlagsToLog.Add(TEXT(#FlagName)); \
+	}
+
+		TRY_LOG_FLAG(didChangeIdentifier);
+		TRY_LOG_FLAG(didChangeResolvedPath);
+		TRY_LOG_FLAG(didReplaceContent);
+		TRY_LOG_FLAG(didReloadContent);
+		TRY_LOG_FLAG(didReorderChildren);
+		TRY_LOG_FLAG(didReorderProperties);
+		TRY_LOG_FLAG(didRename);
+		TRY_LOG_FLAG(didChangePrimVariantSets);
+		TRY_LOG_FLAG(didChangePrimInheritPaths);
+		TRY_LOG_FLAG(didChangePrimSpecializes);
+		TRY_LOG_FLAG(didChangePrimReferences);
+		TRY_LOG_FLAG(didChangeAttributeTimeSamples);
+		TRY_LOG_FLAG(didChangeAttributeConnection);
+		TRY_LOG_FLAG(didChangeRelationshipTargets);
+		TRY_LOG_FLAG(didAddTarget);
+		TRY_LOG_FLAG(didRemoveTarget);
+		TRY_LOG_FLAG(didAddInertPrim);
+		TRY_LOG_FLAG(didAddNonInertPrim);
+		TRY_LOG_FLAG(didRemoveInertPrim);
+		TRY_LOG_FLAG(didRemoveNonInertPrim);
+		TRY_LOG_FLAG(didAddPropertyWithOnlyRequiredFields);
+		TRY_LOG_FLAG(didAddProperty);
+		TRY_LOG_FLAG(didRemovePropertyWithOnlyRequiredFields);
+		TRY_LOG_FLAG(didRemoveProperty);
+
+#undef TRY_LOG_FLAG
+		for (const FString& Flag : FlagsToLog)
+		{
+			UE_LOG(LogUsd, Log, TEXT("%s\t\t'%s'"), *Indent, *Flag);
+		}
+	}
+
+	void LogConvertedChangeListEntry(const UsdUtils::FSdfChangeListEntry& Entry, int32 IndentLevel)
+	{
+		FString Indent;
+		for (int Index = 0; Index < IndentLevel; ++Index)
+		{
+			Indent += TEXT("\t");
+		}
+
+		UE_LOG(LogUsd, Log, TEXT("%sConverted ChangeListEntry:"), *Indent);
+
+		UE_LOG(LogUsd, Log, TEXT("%s\tAttribute changes:"), *Indent);
+		for (const UsdUtils::FAttributeChange& AttributeChange : Entry.AttributeChanges)
+		{
+			UE_LOG(
+				LogUsd,
+				Log,
+				TEXT("%s\t\t'%s' field '%s': From '%s' to '%s'"),
+				*Indent,
+				*AttributeChange.PropertyName,
+				*AttributeChange.Field,
+				UTF8_TO_TCHAR(pxr::TfStringify(AttributeChange.OldValue.GetUsdValue()).c_str()),
+				UTF8_TO_TCHAR(pxr::TfStringify(AttributeChange.NewValue.GetUsdValue()).c_str())
+			);
+		}
+
+		UE_LOG(LogUsd, Log, TEXT("%s\tSubLayerChanges:"), *Indent);
+		static_assert(static_cast<int>(UsdUtils::ESubLayerChangeType::SubLayerAdded) == 0, "Enum values changed!");
+		static_assert(static_cast<int>(UsdUtils::ESubLayerChangeType::SubLayerRemoved) == 1, "Enum values changed!");
+		static_assert(static_cast<int>(UsdUtils::ESubLayerChangeType::SubLayerOffset) == 2, "Enum values changed!");
+		const TCHAR* SubLayerChangeTypeStr[] = {
+			TEXT("SubLayerAdded"),
+			TEXT("SubLayerRemoved"),
+			TEXT("SubLayerOffset"),
+		};
+		for (const TPair<FString, UsdUtils::ESubLayerChangeType>& SubLayerChange : Entry.SubLayerChanges)
+		{
+			UE_LOG(
+				LogUsd,
+				Log,
+				TEXT("%s\t\t'%s': change type '%s'"),
+				*Indent,
+				*SubLayerChange.Key,
+				SubLayerChangeTypeStr[static_cast<int>(SubLayerChange.Value)]
+			);
+		}
+
+		UE_LOG(LogUsd, Log, TEXT("%s\tOldPath: '%s'"), *Indent, *Entry.OldPath);
+
+		UE_LOG(LogUsd, Log, TEXT("%s\tOldIdentifier: '%s'"), *Indent, *Entry.OldIdentifier);
+
+		UE_LOG(LogUsd, Log, TEXT("%s\tFlags:"), *Indent);
+		TArray<FString> FlagsToLog;
+#define TRY_LOG_FLAG(FlagName)           \
+	if (Entry.Flags.##FlagName)          \
+	{                                    \
+		FlagsToLog.Add(TEXT(#FlagName)); \
+	}
+
+		TRY_LOG_FLAG(bDidChangeIdentifier);
+		TRY_LOG_FLAG(bDidChangeResolvedPath);
+		TRY_LOG_FLAG(bDidReplaceContent);
+		TRY_LOG_FLAG(bDidReloadContent);
+		TRY_LOG_FLAG(bDidReorderChildren);
+		TRY_LOG_FLAG(bDidReorderProperties);
+		TRY_LOG_FLAG(bDidRename);
+		TRY_LOG_FLAG(bDidChangePrimVariantSets);
+		TRY_LOG_FLAG(bDidChangePrimInheritPaths);
+		TRY_LOG_FLAG(bDidChangePrimSpecializes);
+		TRY_LOG_FLAG(bDidChangePrimReferences);
+		TRY_LOG_FLAG(bDidChangeAttributeTimeSamples);
+		TRY_LOG_FLAG(bDidChangeAttributeConnection);
+		TRY_LOG_FLAG(bDidChangeRelationshipTargets);
+		TRY_LOG_FLAG(bDidAddTarget);
+		TRY_LOG_FLAG(bDidRemoveTarget);
+		TRY_LOG_FLAG(bDidAddInertPrim);
+		TRY_LOG_FLAG(bDidAddNonInertPrim);
+		TRY_LOG_FLAG(bDidRemoveInertPrim);
+		TRY_LOG_FLAG(bDidRemoveNonInertPrim);
+		TRY_LOG_FLAG(bDidAddPropertyWithOnlyRequiredFields);
+		TRY_LOG_FLAG(bDidAddProperty);
+		TRY_LOG_FLAG(bDidRemovePropertyWithOnlyRequiredFields);
+		TRY_LOG_FLAG(bDidRemoveProperty);
+
+#undef TRY_LOG_FLAG
+		for (const FString& Flag : FlagsToLog)
+		{
+			UE_LOG(LogUsd, Log, TEXT("%s\t\t'%s'"), *Indent, *Flag);
+		}
+	}
+
+	void LogObjectsChangedPathRange(const pxr::UsdNotice::ObjectsChanged::PathRange& PathRange)
+	{
+		FScopedUsdAllocs Allocs;
+
+		for (pxr::UsdNotice::ObjectsChanged::PathRange::const_iterator It = PathRange.begin(); It != PathRange.end(); ++It)
+		{
+			const FString FullFieldPath = UTF8_TO_TCHAR(It->GetAsString().c_str());
+			UE_LOG(LogUsd, Log, TEXT("\t\tObject '%s'"), *FullFieldPath);
+
+			const std::vector<const pxr::SdfChangeList::Entry*>& Changes = It.base()->second;
+			for (const pxr::SdfChangeList::Entry* Entry : Changes)
+			{
+				if (Entry)
+				{
+					const int32 IndentLevel = 3;
+					LogChangeListEntry(*Entry, IndentLevel);
+				}
+				else
+				{
+					UE_LOG(LogUsd, Log, TEXT("\t\t\tNullptr change"));
+				}
+			}
+		}
+	}
+
+	void LogConvertedChangesByPath(const UsdUtils::FObjectChangesByPath& Changes)
+	{
+		for (const TPair<FString, TArray<UsdUtils::FSdfChangeListEntry>>& Pair : Changes)
+		{
+			UE_LOG(LogUsd, Log, TEXT("\t\tObject '%s'"), *Pair.Key);
+
+			for (const UsdUtils::FSdfChangeListEntry& ChangeListEntry : Pair.Value)
+			{
+				const int32 IndentLevel = 3;
+				LogConvertedChangeListEntry(ChangeListEntry, IndentLevel);
+			}
+		}
+	}
+
+	void LogNotice(const pxr::UsdNotice::ObjectsChanged& Notice, const pxr::UsdStageWeakPtr& Sender, int32 BlockCounter)
+	{
+		FScopedUsdAllocs Allocs;
+
+		UE_LOG(
+			LogUsd,
+			Warning,
+			TEXT("pxr::UsdNotice::ObjectsChanged from sender '%s' (blocked? %d):"),
+			Sender ? UTF8_TO_TCHAR(Sender->GetRootLayer()->GetIdentifier().c_str()) : TEXT(""),
+			BlockCounter > 0
+		);
+
+		UE_LOG(LogUsd, Log, TEXT("\tInfoChanges:"));
+		LogObjectsChangedPathRange(Notice.GetChangedInfoOnlyPaths());
+
+		UE_LOG(LogUsd, Log, TEXT("\tResyncChanges:"));
+		LogObjectsChangedPathRange(Notice.GetResyncedPaths());
+
+		UE_LOG(LogUsd, Log, TEXT("\tResolvedAssetPaths:"));
+		LogObjectsChangedPathRange(Notice.GetResolvedAssetPathsResyncedPaths());
+	}
+
+	void LogNotice(const UsdUtils::FObjectChangesByPath& ConvertedInfoChanges, const UsdUtils::FObjectChangesByPath& ConvertedResyncChanges)
+	{
+		UE_LOG(LogUsd, Warning, TEXT("Converted ObjectChange notice:"));
+
+		UE_LOG(LogUsd, Log, TEXT("\tConverted InfoChanges:"));
+		LogConvertedChangesByPath(ConvertedInfoChanges);
+		UE_LOG(LogUsd, Log, TEXT("\tConverted ResyncChanges:"));
+		LogConvertedChangesByPath(ConvertedResyncChanges);
+	}
+
+	void LogNotice(const pxr::UsdNotice::StageEditTargetChanged& Notice, const pxr::UsdStageWeakPtr& Sender)
+	{
+		UE_LOG(
+			LogUsd,
+			Warning,
+			TEXT("pxr::UsdNotice::StageEditTargetChanged from sender '%s':"),
+			Sender ? UTF8_TO_TCHAR(Sender->GetRootLayer()->GetIdentifier().c_str()) : TEXT("")
+		);
+	}
+
+	void LogNotice(const pxr::SdfNotice::LayersDidChange& Notice, int32 BlockCounter)
+	{
+		FScopedUsdAllocs Allocs;
+
+		UE_LOG(LogUsd, Warning, TEXT("pxr::SdfNotice::LayersDidChange (blocked? %d)"), BlockCounter > 0);
+
+		UE_LOG(LogUsd, Log, TEXT("\tSerial number: '%d'"), Notice.GetSerialNumber());
+
+		std::vector<std::pair<pxr::SdfLayerHandle, pxr::SdfChangeList>> ChangeListVec = Notice.GetChangeListVec();
+		for (const std::pair<pxr::SdfLayerHandle, pxr::SdfChangeList>& Pair : ChangeListVec)
+		{
+			UE_LOG(LogUsd, Log, TEXT("\tLayer: '%s'"), UTF8_TO_TCHAR(Pair.first->GetIdentifier().c_str()));
+
+			for (const std::pair<pxr::SdfPath, pxr::SdfChangeList::Entry>& ObjectPair : Pair.second.GetEntryList())
+			{
+				UE_LOG(LogUsd, Log, TEXT("\t\tObject: '%s'"), UTF8_TO_TCHAR(ObjectPair.first.GetString().c_str()));
+
+				const int32 IndentLevel = 3;
+				LogChangeListEntry(ObjectPair.second, IndentLevel);
+			}
+		}
+	}
+
+	void LogNotice(const pxr::SdfNotice::LayerDirtinessChanged& Notice, int32 BlockCounter)
+	{
+		UE_LOG(LogUsd, Warning, TEXT("pxr::SdfNotice::LayerDirtinessChanged (blocked? %d)"), BlockCounter > 0);
+	}
+
+	void LogNotice(const UsdUtils::FLayerToSdfChangeList& ConvertedLayerToChangeList)
+	{
+		FScopedUsdAllocs Allocs;
+
+		UE_LOG(LogUsd, Warning, TEXT("Converted LayerChanges:"));
+
+		for (const TPair<UE::FSdfLayerWeak, UsdUtils::FSdfChangeList>& ConvertedPair : ConvertedLayerToChangeList)
+		{
+			UE_LOG(LogUsd, Log, TEXT("\tLayer: '%s'"), *ConvertedPair.Key.GetIdentifier());
+
+			for (const TPair<UE::FSdfPath, UsdUtils::FSdfChangeListEntry>& EntryPair : ConvertedPair.Value)
+			{
+				UE_LOG(LogUsd, Log, TEXT("\tObject: '%s'"), *EntryPair.Key.GetString());
+
+				const int32 IndentLevel = 2;
+				LogConvertedChangeListEntry(EntryPair.Value, IndentLevel);
+			}
+		}
+	}
+#endif	  // ENABLE_NOTICE_LOGGING
+
 	void ConvertSdfChangeListEntry(
 		const pxr::SdfPath& ChangedObject,
 		const pxr::SdfChangeList::Entry& Entry,
@@ -271,7 +592,7 @@ namespace UsdToUnreal
 		return true;
 	}
 #endif	  // USE_USD_SDK
-}	 // namespace UsdToUnreal
+}	 // namespace UE::USDListener::Private
 
 class FUsdListenerImpl
 #if USE_USD_SDK
@@ -436,6 +757,10 @@ PRAGMA_ENABLE_DEPRECATION_WARNINGS
 #if USE_USD_SDK
 void FUsdListenerImpl::HandleObjectsChangedNotice(const pxr::UsdNotice::ObjectsChanged& Notice, const pxr::UsdStageWeakPtr& Sender)
 {
+#if ENABLE_NOTICE_LOGGING
+	UE::USDListener::Private::LogNotice(Notice, Sender, IsBlocked.GetValue());
+#endif	  // ENABLE_NOTICE_LOGGING
+
 	if (!OnObjectsChanged.IsBound())
 	{
 		return;
@@ -448,22 +773,34 @@ void FUsdListenerImpl::HandleObjectsChangedNotice(const pxr::UsdNotice::ObjectsC
 
 	UsdUtils::FObjectChangesByPath InfoChanges;
 	UsdUtils::FObjectChangesByPath ResyncChanges;
-	UsdToUnreal::ConvertObjectsChangedNotice(Notice, InfoChanges, ResyncChanges);
+	UE::USDListener::Private::ConvertObjectsChangedNotice(Notice, InfoChanges, ResyncChanges);
 	if (InfoChanges.Num() > 0 || ResyncChanges.Num() > 0)
 	{
 		FScopedUnrealAllocs UnrealAllocs;
 		OnObjectsChanged.Broadcast(InfoChanges, ResyncChanges);
 	}
+
+#if ENABLE_NOTICE_LOGGING
+	UE::USDListener::Private::LogNotice(InfoChanges, ResyncChanges);
+#endif	  // ENABLE_NOTICE_LOGGING
 }
 
 void FUsdListenerImpl::HandleStageEditTargetChangedNotice(const pxr::UsdNotice::StageEditTargetChanged& Notice, const pxr::UsdStageWeakPtr& Sender)
 {
+#if ENABLE_NOTICE_LOGGING
+	UE::USDListener::Private::LogNotice(Notice, Sender);
+#endif	  // ENABLE_NOTICE_LOGGING
+
 	FScopedUnrealAllocs UnrealAllocs;
 	OnStageEditTargetChanged.Broadcast();
 }
 
 void FUsdListenerImpl::HandleLayersChangedNotice(const pxr::SdfNotice::LayersDidChange& Notice)
 {
+#if ENABLE_NOTICE_LOGGING
+	UE::USDListener::Private::LogNotice(Notice, IsBlocked.GetValue());
+#endif	  // ENABLE_NOTICE_LOGGING
+
 	PRAGMA_DISABLE_DEPRECATION_WARNINGS
 	if ((!OnLayersChanged.IsBound() && !OnSdfLayersChanged.IsBound()) || IsBlocked.GetValue() > 0)
 	{
@@ -491,7 +828,7 @@ void FUsdListenerImpl::HandleLayersChangedNotice(const pxr::SdfNotice::LayersDid
 			{
 				TPair<UE::FSdfPath, UsdUtils::FSdfChangeListEntry>& ConvertedChange = ConvertedChangeList.Value.Emplace_GetRef();
 				ConvertedChange.Key = UE::FSdfPath{Change.first};
-				UsdToUnreal::ConvertSdfChangeListEntry(Change.first, Change.second, ConvertedChange.Value);
+				UE::USDListener::Private::ConvertSdfChangeListEntry(Change.first, Change.second, ConvertedChange.Value);
 
 				const pxr::SdfChangeList::Entry::_Flags& Flags = Change.second.flags;
 				if (Flags.didReloadContent)
@@ -502,6 +839,10 @@ void FUsdListenerImpl::HandleLayersChangedNotice(const pxr::SdfNotice::LayersDid
 		}
 	}
 
+#if ENABLE_NOTICE_LOGGING
+	UE::USDListener::Private::LogNotice(ConvertedLayerToChangeList);
+#endif	  // ENABLE_NOTICE_LOGGING
+
 	FScopedUnrealAllocs UnrealAllocs;
 	OnLayersChanged.Broadcast(LayersNames);
 	OnSdfLayersChanged.Broadcast(ConvertedLayerToChangeList);
@@ -510,6 +851,10 @@ void FUsdListenerImpl::HandleLayersChangedNotice(const pxr::SdfNotice::LayersDid
 
 void FUsdListenerImpl::HandleLayerDirtinessChangedNotice(const pxr::SdfNotice::LayerDirtinessChanged& Notice)
 {
+#if ENABLE_NOTICE_LOGGING
+	UE::USDListener::Private::LogNotice(Notice, IsBlocked.GetValue());
+#endif	  // ENABLE_NOTICE_LOGGING
+
 	if (!OnSdfLayerDirtinessChanged.IsBound() || IsBlocked.GetValue() > 0)
 	{
 		return;
