USD: Fix AssetCache3 never cleaning up GeometryCaches as the asset cache's own reference to the asset was being counted as an "external reference".

diff --git a/Engine/Plugins/Runtime/USDCore/Source/USDClasses/Private/USDAssetCache3.cpp b/Engine/Plugins/Runtime/USDCore/Source/USDClasses/Private/USDAssetCache3.cpp
--- a/Engine/Plugins/Runtime/USDCore/Source/USDClasses/Private/USDAssetCache3.cpp
+++ b/Engine/Plugins/Runtime/USDCore/Source/USDClasses/Private/USDAssetCache3.cpp
@@ -14,6 +14,8 @@
 #include "ComponentRecreateRenderStateContext.h"
 #include "Containers/Ticker.h"
 #include "Engine/Blueprint.h"
+#include "GeometryCache.h"
+#include "GeometryCacheTrack.h"
 #include "GroomAsset.h"
 #include "HAL/FileManager.h"
 #include "Interfaces/Interface_AssetUserData.h"
@@ -80,7 +82,7 @@ namespace UE::USDAssetCache3::Private
 	// This should only be called with memory-only objects (i.e. don't call this with assets that have been previously saved to disk).
 	// This aims to be more safe and fast than comprehensive: We'll run GC at the end which can wipe some references on its own, so even if
 	// we leave some objects behind there's always the chance that closing the next stage may clear them anyway
-	void SafeDeleteObjects(const TSet<UObject*>& ObjectsToDelete)
+	void SafeDeleteObjects(const TSet<UObject*>& ObjectsToDelete, const UUsdAssetCache3& AssetCache)
 	{
 		TRACE_CPUPROFILER_EVENT_SCOPE(UUsdAssetCache3::SafeDeleteObjects);
 
@@ -217,6 +219,21 @@ namespace UE::USDAssetCache3::Private
 								continue;
 							}
 
+							// Manually ignore the references from UGeometryCacheTrack to UGeometryCaches. The tracks are owned by the caches
+							// themselves, they just happen to have the transient package as their outer instead of the UGeometryCache asset,
+							// so our mechanism here considers them external referencers
+							if (Cast<UGeometryCache>(ObjectToDelete) && Cast<UGeometryCacheTrack>(Referencer))
+							{
+								continue;
+							}
+
+							// GeometryCaches can be fully owned by the asset cache now, so they would count as referencers here.
+							// Of course, we don't care about those references either
+							if (Referencer == &AssetCache)
+							{
+								continue;
+							}
+
 							// Quickly ignore references from actors and components on the preview world (i.e. thumbnail renderers, asset editors,
 							// etc...)
 							if (Referencer->HasAnyFlags(RF_DefaultSubObject))
@@ -1018,7 +1035,7 @@ void UUsdAssetCache3::DeleteUnreferencedAssets(bool bShowConfirmation)
 		//    to clear anything anymore, which defeats the entire purpose
 		// Note that SafeDeleteObjects should be a bit slower than ObjectTools::DeleteAssets though, unfortunately...
 		ObjectsToDelete.Remove(nullptr);
-		UE::USDAssetCache3::Private::SafeDeleteObjects(ObjectsToDelete);
+		UE::USDAssetCache3::Private::SafeDeleteObjects(ObjectsToDelete, *this);
 	}
 
 	FWriteScopeLock Lock(RWLock);
