USD: Fix draw mode extents animation not being written out to USD.

diff --git a/Engine/Plugins/Runtime/USDCore/Source/USDUtilities/Private/USDPrimConversion.cpp b/Engine/Plugins/Runtime/USDCore/Source/USDUtilities/Private/USDPrimConversion.cpp
--- a/Engine/Plugins/Runtime/USDCore/Source/USDUtilities/Private/USDPrimConversion.cpp
+++ b/Engine/Plugins/Runtime/USDCore/Source/USDUtilities/Private/USDPrimConversion.cpp
@@ -4868,14 +4868,18 @@ bool UnrealToUsd::CreateComponentPropertyBaker(
 			pxr::UsdAttribute ExtentsAttr;
 			if (pxr::UsdGeomBoundable Boundable{UsdPrim})
 			{
-				// Try using the extents attribute if we already have one authored
-				ExtentsAttr = Boundable.GetExtentAttr();
+				// Try using the extents attribute if we're a boundable
+				ExtentsAttr = Boundable.CreateExtentAttr();
 			}
 			if (!ExtentsAttr)
 			{
+				// Otherwise fallback to trying to use extentsHint
 				if (pxr::UsdGeomModelAPI GeomModelAPI = pxr::UsdGeomModelAPI::Apply(UsdPrim))
 				{
-					ExtentsAttr = GeomModelAPI.GetExtentsHintAttr();
+					// Copied from the implementation of pxr::UsdGeomModelAPI::SetExtentsHint() because for whatever reason
+					// there is no CreateExtentsHintAttr()
+					const bool bCustom = false;
+					ExtentsAttr = UsdPrim.CreateAttribute(pxr::UsdGeomTokens->extentsHint, pxr::SdfValueTypeNames->Float3Array, bCustom);
 				}
 			}
 			if (!ExtentsAttr)
