USD: Fix duplicate morph targets when importing skeletal meshes.

diff --git a/Engine/Plugins/Runtime/USDCore/Source/USDUtilities/Private/USDSkeletalDataConversion.cpp b/Engine/Plugins/Runtime/USDCore/Source/USDUtilities/Private/USDSkeletalDataConversion.cpp
--- a/Engine/Plugins/Runtime/USDCore/Source/USDUtilities/Private/USDSkeletalDataConversion.cpp
+++ b/Engine/Plugins/Runtime/USDCore/Source/USDUtilities/Private/USDSkeletalDataConversion.cpp
@@ -606,154 +606,6 @@ namespace UsdToUnrealImpl
 			UnmorphedShape.NormalsPerVertex[VertexIndex] = NormalsPerIndex[IndexIndex];
 		}
 	}
-
-	void CreateMorphTargets(
-		UsdUtils::FBlendShapeMap& BlendShapes,
-		const TArray<FSkeletalMeshImportData>& LODIndexToSkeletalMeshImportData,
-		USkeletalMesh* SkeletalMesh
-	)
-	{
-		TRACE_CPUPROFILER_EVENT_SCOPE(USDSkeletalDataConversion::CreateMorphTargets);
-
-		FSkeletalMeshModel* ImportedResource = SkeletalMesh->GetImportedModel();
-		if (LODIndexToSkeletalMeshImportData.Num() != ImportedResource->LODModels.Num())
-		{
-			return;
-		}
-
-		int32 NumLODs = ImportedResource->LODModels.Num();
-
-		// Temporarily pack the mesh data into a different format that MeshUtilities can use
-		TArray<SkelDataConversionImpl::FMeshDataBundle> TempMeshBundlesPerLOD;
-		TempMeshBundlesPerLOD.Reserve(NumLODs);
-
-		TArray<TMap<int32, TArray<int32>>> OrigIndexToBuiltIndicesPerLOD;
-		OrigIndexToBuiltIndicesPerLOD.Reserve(NumLODs);
-
-		TArray<bool> AreNormalsComputedForLODIndex;
-		AreNormalsComputedForLODIndex.SetNumZeroed(NumLODs);
-
-		for (int32 LODIndex = 0; LODIndex < NumLODs; ++LODIndex)
-		{
-			FSkeletalMeshLODModel& LODModel = ImportedResource->LODModels[LODIndex];
-
-			// BuildSkeletalMesh may create/remove vertices, and reorder/optimize the index buffers. We can use MeshToImportVertexMap
-			// to go from new vertex index to original vertex index. Here we invert this map, because our FMorphTargetDeltas all refer
-			// to original vertex indices, so we'll need to map them to the post-build vertex indices
-			TArray<int32>& BuildIndexToOrigIndex = LODModel.MeshToImportVertexMap;
-			TMap<int32, TArray<int32>>& OrigIndexToBuiltIndices = OrigIndexToBuiltIndicesPerLOD.Emplace_GetRef();
-			OrigIndexToBuiltIndices.Reserve(BuildIndexToOrigIndex.Num());
-			for (int32 BuiltIndex = 0; BuiltIndex < BuildIndexToOrigIndex.Num(); ++BuiltIndex)
-			{
-				int32 OrigIndex = BuildIndexToOrigIndex[BuiltIndex];
-				OrigIndexToBuiltIndices.FindOrAdd(OrigIndex).Add(BuiltIndex);
-			}
-
-			SkelDataConversionImpl::FMeshDataBundle& LODMeshBundle = TempMeshBundlesPerLOD.Emplace_GetRef();
-			SkelDataConversionImpl::ConvertImportDataToMeshData(LODIndexToSkeletalMeshImportData[LODIndex], LODMeshBundle);
-		}
-
-		struct FMorphTargetJob
-		{
-			UsdUtils::FUsdBlendShape* BlendShape = nullptr;
-			UMorphTarget* MorphTarget = nullptr;
-		};
-
-		TArray<FMorphTargetJob> MorphTargetJobs;
-		MorphTargetJobs.Reserve(BlendShapes.Num());
-
-		bool bHasValidMorphTarget = false;
-		for (TPair<FString, UsdUtils::FUsdBlendShape>& BlendShapeByPath : BlendShapes)
-		{
-			TRACE_CPUPROFILER_EVENT_SCOPE(USDSkeletalDataConversion::CreatingMorphTargetObject);
-
-			UsdUtils::FUsdBlendShape& BlendShape = BlendShapeByPath.Value;
-			if (!BlendShape.IsValid())
-			{
-				continue;
-			}
-
-			bHasValidMorphTarget = true;
-
-			FMorphTargetJob& Job = MorphTargetJobs.Emplace_GetRef();
-			Job.BlendShape = &BlendShape;
-			Job.MorphTarget = NewObject<UMorphTarget>(SkeletalMesh, *BlendShape.Name);
-			Job.MorphTarget->BaseSkelMesh = SkeletalMesh;
-			SkeletalMesh->GetMorphTargets().Add(Job.MorphTarget);
-
-			// Recompute normals for the final morphed shape in case it doesn't have authored normals
-			// This is required or else the morphed shape will reuse the unmorphed normals, and lighting may look incorrect for
-			// aggressive morph targets.
-			// Note that this should happen *before* we call UpdatesDeltasToMeshBuild, because our MeshDataBundle refers to import data,
-			// and so should our BlendShape
-			for (int32 LODIndex : Job.BlendShape->LODIndicesThatUseThis)
-			{
-				if (!Job.BlendShape->bHasAuthoredTangents)
-				{
-					SkelDataConversionImpl::FMeshDataBundle& UnmorphedShape = TempMeshBundlesPerLOD[LODIndex];
-					if (!AreNormalsComputedForLODIndex[LODIndex])
-					{
-						ComputeSourceNormals(UnmorphedShape);
-						AreNormalsComputedForLODIndex[LODIndex] = true;
-					}
-				}
-			}
-		}
-
-		const int32 MinBatchSize = 1;
-
-		ParallelFor(
-			TEXT("CreateMorphTarget"),
-			MorphTargetJobs.Num(),
-			MinBatchSize,
-			[&MorphTargetJobs, &OrigIndexToBuiltIndicesPerLOD, &TempMeshBundlesPerLOD, ImportedResource](int32 Index)
-			{
-				TRACE_CPUPROFILER_EVENT_SCOPE(USDSkeletalDataConversion::CreateMorphTargetJob);
-
-				FMorphTargetJob& Job = MorphTargetJobs[Index];
-				if (!Job.BlendShape || !Job.MorphTarget)
-				{
-					return;
-				}
-
-				for (int32 LODIndex : Job.BlendShape->LODIndicesThatUseThis)
-				{
-					const SkelDataConversionImpl::FMeshDataBundle& UnmorphedShape = TempMeshBundlesPerLOD[LODIndex];
-
-					if (!Job.BlendShape->bHasAuthoredTangents)
-					{
-						SkelDataConversionImpl::ComputeTangentDeltas(UnmorphedShape, *Job.BlendShape);
-					}
-
-					TArray<FMorphTargetDelta> Vertices;
-					if (Job.BlendShape->LODIndicesThatUseThis.Num() > 1)
-					{
-						// Need to copy this here because different LODs may build differently, and so the deltas may need to be updated differently
-						Vertices = Job.BlendShape->Vertices;
-					}
-					else
-					{
-						Vertices = MoveTemp(Job.BlendShape->Vertices);
-					}
-
-					SkelDataConversionImpl::UpdatesDeltasToMeshBuild(Vertices, OrigIndexToBuiltIndicesPerLOD[LODIndex]);
-
-					const bool bCompareNormal = true;
-					FSkeletalMeshLODModel& LODModel = ImportedResource->LODModels[LODIndex];
-					Job.MorphTarget->PopulateDeltas(Vertices, LODIndex, LODModel.Sections, bCompareNormal);
-				}
-
-				// Don't need this data anymore as it has been moved into UMorphTarget
-				Job.BlendShape->Vertices.Empty();
-			}
-		);
-
-		if (bHasValidMorphTarget)
-		{
-			SkeletalMesh->MarkPackageDirty();
-			SkeletalMesh->InitMorphTargetsAndRebuildRenderData();
-		}
-	}
 }	 // namespace UsdToUnrealImpl
 
 namespace UnrealToUsdImpl
@@ -2966,27 +2818,49 @@ bool UsdToUnreal::ConvertSkeletalImportData(
 			return false;
 		}
 
-		// We must also provide the ImportData with morph target information now.
-		// Reference: FSkeletalMeshImportData::AddMorphTarget
-		// We don't use that function directly as matching the interface would involve copying our BlendShape.Vertices
-		// into a new FMorphTargetLODModel.
+		// UMorphTarget::PopulateDeltas called by BuildMorphTargetsInternal will ignore deltas below
+		// MorphThresholdPosition, so we must have something larger than that
+		const FVector3f SmallMorphDelta{0.0f, 0.0f, FMath::Max(BuildOptions.OverlappingThresholds.MorphThresholdPosition * 1.1f, 1e-4f)};
+
+		// Morph target data is now primarily provided via the MeshDescription. For now we still don't convert
+		// skeletal data directly to the MeshDescription, so we must feed it into FSkeletalMeshImportData, so that
+		// SaveLODImportedData converts it into the skeletal MeshDescription for us
+		//
+		// Reference: FSkeletalMeshImportData::AddMorphTarget, but we don't use it directly as matching the interface
+		// would involve copying our BlendShape.Vertices into a new FMorphTargetLODModel.
+		//
 		// TODO: Add in the morph target normal data (from FMorphTargetDelta::TangentZDelta) to the import data
-		// at the right location when FSkeletalMeshImportData::GetMeshDescription starts reading normal data from that location
+		// at the right location when FSkeletalMeshImportData::GetMeshDescription starts reading normal data, or whenever
+		// we start converting skeletal data into MeshDescriptions
 		LODImportData.MorphTargets.Reserve(InBlendShapesByPath.Num());
 		LODImportData.MorphTargetNames.Reserve(InBlendShapesByPath.Num());
 		LODImportData.MorphTargetModifiedPoints.Reserve(InBlendShapesByPath.Num());
 		for (const TPair<FString, UsdUtils::FUsdBlendShape>& Pair : InBlendShapesByPath)
 		{
 			const UsdUtils::FUsdBlendShape& BlendShape = Pair.Value;
-			if (!BlendShape.LODIndicesThatUseThis.Contains(LODIndex))
-			{
-				continue;
-			}
 
+			// The morph targets used for higher LOD levels must be a subset of the morph targets used for
+			// lower LODS, or else the skeletal mesh build will just discard them. So even if this blend
+			// shape doesn't affect this LOD, just add it anyway
 			LODImportData.MorphTargetNames.Add(BlendShape.Name);
 			FSkeletalMeshImportData& MorphTarget = LODImportData.MorphTargets.Emplace_GetRef();
 			TSet<uint32>& NewModifiedPoints = LODImportData.MorphTargetModifiedPoints.Emplace_GetRef();
 
+			if (!BlendShape.LODIndicesThatUseThis.Contains(LODIndex))
+			{
+				if (LODImportData.Points.Num() > 0)
+				{
+					// Additionally, in order to keep a morph target that in USD is only defined for higher LODS, we must
+					// add a tiny valid delta when handling it for lower LODs, or else the skeletal mesh build and
+					// UMorphTarget::PopulateDeltas will get rid of the morph target during the processing of the lower LODS,
+					// and then ignore it when later processing the higher LODs, as the morph target was removed...
+					MorphTarget.Points.Add(LODImportData.Points[0] + SmallMorphDelta);
+					NewModifiedPoints.Add(0);
+				}
+
+				continue;
+			}
+
 			MorphTarget.Points.Reserve(LODImportData.Points.Num());
 			NewModifiedPoints.Reserve(BlendShape.Vertices.Num());
 
@@ -3016,8 +2890,6 @@ bool UsdToUnreal::ConvertSkeletalImportData(
 		return false;
 	}
 
-	UsdToUnrealImpl::CreateMorphTargets(InBlendShapesByPath, InLODIndexToSkeletalMeshImportData, InOutSkeletalMesh);
-
 	// "Declare" the morph target curves on the skeleton or skeletal mesh according to bAddCurveMetadataToSkeleton.
 	// This is important otherwise the ControlRig will not hoist these curves as controls when using e.g. FKControlRig.
 	for (const TPair<FString, UsdUtils::FUsdBlendShape>& BlendShapeByPath : InBlendShapesByPath)
