USD: Only clean up the asset cache when actively closing stages.

diff --git a/Engine/Plugins/Importers/USDImporter/Source/USDStage/Private/USDStageActor.cpp b/Engine/Plugins/Importers/USDImporter/Source/USDStage/Private/USDStageActor.cpp
--- a/Engine/Plugins/Importers/USDImporter/Source/USDStage/Private/USDStageActor.cpp
+++ b/Engine/Plugins/Importers/USDImporter/Source/USDStage/Private/USDStageActor.cpp
@@ -4995,8 +4995,6 @@ bool AUsdStageActor::UnloadAssets(const UE::FSdfPath& StartPrimPath, bool bForEn
 		}
 	}
 
-	AssetCache->RequestDelayedAssetAutoCleanup();
-
 	return bAssetsDiscarded;
 }
 
diff --git a/Engine/Plugins/Runtime/USDCore/Source/USDClasses/Public/USDAssetCache3.h b/Engine/Plugins/Runtime/USDCore/Source/USDClasses/Public/USDAssetCache3.h
--- a/Engine/Plugins/Runtime/USDCore/Source/USDClasses/Public/USDAssetCache3.h
+++ b/Engine/Plugins/Runtime/USDCore/Source/USDClasses/Public/USDAssetCache3.h
@@ -374,8 +374,8 @@ public:
 	TMap<FString, FSoftObjectPath> HashToAssetPaths;
 
 	/**
-	 * If this is true, every time a UsdStageActor using this asset cache abandons assets (e.g. when closing a stage or resyncing a prim),
-	 * it will attempt to call DeleteUnreferencedAssets, potentially dropping *any* unreferenced asset, due to this operation or previous ones.
+	 * If this is true, every time a UsdStageActor using this asset cache closes a stage or swaps asset caches it will attempt to call
+	 * DeleteUnreferencedAssets, potentially dropping *any* unreferenced asset, due to this operation or previous ones.
 	 *
 	 * Enable this if you want your AssetDirectory folder to be automatically cleaned up as stages close, and don't plan on keeping
 	 * other external references to those assets.
