USD: Fix AssetCache3 cleaning assets way to aggressively, ignoring some situations where they were still actively being used.

diff --git a/Engine/Plugins/Runtime/USDCore/Source/USDClasses/Private/USDAssetCache3.cpp b/Engine/Plugins/Runtime/USDCore/Source/USDClasses/Private/USDAssetCache3.cpp
--- a/Engine/Plugins/Runtime/USDCore/Source/USDClasses/Private/USDAssetCache3.cpp
+++ b/Engine/Plugins/Runtime/USDCore/Source/USDClasses/Private/USDAssetCache3.cpp
@@ -234,13 +234,6 @@ namespace UE::USDAssetCache3::Private
 								continue;
 							}
 
-							// Quickly ignore references from actors and components on the preview world (i.e. thumbnail renderers, asset editors,
-							// etc...)
-							if (Referencer->HasAnyFlags(RF_DefaultSubObject))
-							{
-								continue;
-							}
-
 							ExternalReferencers.Add(Referencer);
 						}
 					}
