USD: Don't show some AssetCache3 warnings (related to invalid asset paths) unless necessary.

diff --git a/Engine/Plugins/Runtime/USDCore/Source/USDClasses/Private/USDAssetCache3.cpp b/Engine/Plugins/Runtime/USDCore/Source/USDClasses/Private/USDAssetCache3.cpp
--- a/Engine/Plugins/Runtime/USDCore/Source/USDClasses/Private/USDAssetCache3.cpp
+++ b/Engine/Plugins/Runtime/USDCore/Source/USDClasses/Private/USDAssetCache3.cpp
@@ -48,6 +48,10 @@ namespace UE::USDAssetCache3::Private
 			// Check if the package exists on disk first to try and avoid some ugly warnings if we try calling TryLoad with a broken path
 			const FAssetRegistryModule& AssetRegistryModule = FModuleManager::GetModuleChecked<FAssetRegistryModule>(TEXT("AssetRegistry"));
 			FAssetData AssetData = AssetRegistryModule.Get().GetAssetByObjectPath(Path);
+			if (!AssetData.IsValid())
+			{
+				return nullptr;
+			}
 
 			// We can't load objects from disk from an async thread
 			if (!AssetData.IsAssetLoaded())
@@ -600,15 +604,22 @@ UObject* UUsdAssetCache3::GetOrCreateCustomCachedAsset(
 	// Our reverse map is just one to one, so we can't allow associating two hashes to the same asset
 	if (FString* OldHash = AssetPathToHashes.Find(Asset))
 	{
-		UE_LOG(
-			LogUsd,
-			Warning,
-			TEXT("An asset can only be associated with a single hash! Discarding old hash '%s' mapped to recently cached asset '%s' (new hash '%s')"),
-			**OldHash,
-			*Asset->GetPathName(),
-			*Hash
-		);
-		StopTrackingAssetInternal(*OldHash);
+		// There are some scenarios in which we have to recache the same asset to the same hash (UE-214909), and that should
+		// be fine. We only care if we somehow have different hashes for the same asset
+		if (*OldHash != Hash)
+		{
+			UE_LOG(
+				LogUsd,
+				Warning,
+				TEXT(
+					"An asset can only be associated with a single hash! Discarding old hash '%s' mapped to recently cached asset '%s' (new hash '%s')"
+				),
+				**OldHash,
+				*Asset->GetPathName(),
+				*Hash
+			);
+			StopTrackingAssetInternal(*OldHash);
+		}
 	}
 
 	// We don't want to inherit any old referencers in case we happened to have some old data in the referencer
