USD MaterialX ShaderGraph:
- Fix wrong recurrence in the shadergraph, we should stop after assigning the different nodenames or values to the inputs not before
- Fix a crash that may occur when dealing with type string in the ShaderGraph, surrounding the allocation with a TUsdStore wrapper to avoid the crash caused by the C Runtime Library across the different dlls.

diff --git a/Engine/Plugins/Runtime/USDCore/Source/USDUtilities/Private/USDShadeConversion.cpp b/Engine/Plugins/Runtime/USDCore/Source/USDUtilities/Private/USDShadeConversion.cpp
--- a/Engine/Plugins/Runtime/USDCore/Source/USDUtilities/Private/USDShadeConversion.cpp
+++ b/Engine/Plugins/Runtime/USDCore/Source/USDUtilities/Private/USDShadeConversion.cpp
@@ -2239,24 +2239,30 @@ namespace UE::UsdShadeConversion::Private
 				}
 
 				std::string ConnectedShaderName = ConnectedShader.GetPrim().GetName().GetString();
-
-				// Add the node only once
-				if (!Document->getNode(ConnectedShaderName))
+				
+				bool bComputeShaderGraph = false;
+				mx::NodePtr ConnectedNode = Document->getNode(ConnectedShaderName);
+				if (!ConnectedNode)
 				{
-					mx::NodePtr ConnectedNode = Document->addNodeInstance(NodeDef, ConnectedShaderName);
-					Input->setConnectedNode(ConnectedNode);
+					ConnectedNode = Document->addNodeInstance(NodeDef, ConnectedShaderName);
+					bComputeShaderGraph = true;
+				}
 
-					// add the output of the connected node to the current input
-					TUsdStore<std::vector<pxr::UsdShadeOutput>> UsdOutputs = ConnectedShader.GetOutputs();
-					for (const pxr::UsdShadeOutput& UsdOutput : *UsdOutputs)
+				Input->setConnectedNode(ConnectedNode);
+
+				// add the output of the connected node to the current input
+				TUsdStore<std::vector<pxr::UsdShadeOutput>> UsdOutputs = ConnectedShader.GetOutputs();
+				for(const pxr::UsdShadeOutput& UsdOutput : *UsdOutputs)
+				{
+					//no need to add the default output
+					if(std::string OutputName = UsdOutput.GetBaseName().GetString(); OutputName != "out")
 					{
-						//no need to add the default output
-						if (std::string OutputName = UsdOutput.GetBaseName().GetString(); OutputName != "out")
-						{
-							Input->setOutputString(OutputName);
-						}
+						Input->setOutputString(OutputName);
 					}
+				}
 
+				if (bComputeShaderGraph)
+				{
 					ComputeShaderGraph(ConnectedNode, ConnectedShader);
 				}
 			}
@@ -2329,7 +2335,17 @@ namespace UE::UsdShadeConversion::Private
 				}
 				else if (InputType == "string")
 				{
-					SetInputValue(std::string{});
+					if (TUsdStore<std::string> Value; UsdInput.Get<std::string>(&Value.Get()))
+					{
+						Input->setValue(Value.Get());
+					}
+					else
+					{
+						mx::NodeDefPtr NodeDef = Node->getNodeDef();
+						mx::InputPtr ActiveInput = NodeDef->getActiveInput(Input->getName());
+						std::string DefaultValue = ActiveInput->getDefaultValue()->asA<std::string>();
+						Input->setValue(DefaultValue);
+					}
 				}
 				else if (InputType == "filename")
 				{
