USD MaterialX:
- Fix a crash only happening on Linux, so far caused by an std::unordered_map most likely freed twice in MaterialX Value.h (_creatorMap static variable), the fix disable MaterialX on Linux when used by USD

diff --git a/Engine/Plugins/Interchange/Runtime/Source/Import/Private/Usd/InterchangeUsdTranslator.cpp b/Engine/Plugins/Interchange/Runtime/Source/Import/Private/Usd/InterchangeUsdTranslator.cpp
--- a/Engine/Plugins/Interchange/Runtime/Source/Import/Private/Usd/InterchangeUsdTranslator.cpp
+++ b/Engine/Plugins/Interchange/Runtime/Source/Import/Private/Usd/InterchangeUsdTranslator.cpp
@@ -803,6 +803,8 @@ namespace UE::InterchangeUsdTranslator::Private
 				return true;
 			}
 
+			// Only enable the shader graph on Windows for the time being as it crashes on Linux cause of a probable double free
+#if PLATFORM_WINDOWS
 			bool bHasShaderGraph = false;			
 			if(MaterialX::DocumentPtr Document = UsdToUnreal::CreateSurfaceMaterial(Prim, TranslatorSettings ? *RenderContext.ToString() : nullptr))
 			{
@@ -825,6 +827,7 @@ namespace UE::InterchangeUsdTranslator::Private
 
 				return bHasShaderGraph;
 			}
+#endif // PLATFORM_WINDOWS
 		}
 #endif // WITH_EDITOR
 		return false;
diff --git a/Engine/Plugins/Runtime/USDCore/Source/USDUtilities/Private/USDShadeConversion.cpp b/Engine/Plugins/Runtime/USDCore/Source/USDUtilities/Private/USDShadeConversion.cpp
--- a/Engine/Plugins/Runtime/USDCore/Source/USDUtilities/Private/USDShadeConversion.cpp
+++ b/Engine/Plugins/Runtime/USDCore/Source/USDUtilities/Private/USDShadeConversion.cpp
@@ -2,7 +2,6 @@
 
 #include "USDShadeConversion.h"
 
-#include "MaterialXFormat/Util.h"
 
 #if USE_USD_SDK
 
@@ -65,13 +64,17 @@
 
 #include "USDIncludesEnd.h"
 
+#if WITH_EDITOR && PLATFORM_WINDOWS
+#include "MaterialXFormat/Util.h"
+
+namespace mx = MaterialX;
+#endif
+
 // Required for packaging
 class UTextureFactory;
 
 #define LOCTEXT_NAMESPACE "USDShadeConversion"
 
-namespace mx = MaterialX;
-
 namespace UE::UsdShadeConversion::Private
 {
 #if WITH_EDITOR
@@ -2142,6 +2145,7 @@ bool UsdToUnreal::FTextureParameterValue::GetSRGBValue() const
 	return bSRGB.Get(bFallbackSRGB);
 }
 
+#if WITH_EDITOR && PLATFORM_WINDOWS
 namespace UE::UsdShadeConversion::Private
 {
 	std::unordered_set<std::string> MaterialXTypes
@@ -2457,6 +2461,7 @@ mx::DocumentPtr UsdToUnreal::CreateSurfaceMaterial(const pxr::UsdPrim& UsdShadeM
 
 	return Document;
 }
+#endif // WITH_EDITOR && PLATFORM_WINDOWS
 
 bool UsdToUnreal::ConvertMaterial(
 	const pxr::UsdPrim& UsdShadeMaterialPrim,
diff --git a/Engine/Plugins/Runtime/USDCore/Source/USDUtilities/Public/USDShadeConversion.h b/Engine/Plugins/Runtime/USDCore/Source/USDUtilities/Public/USDShadeConversion.h
--- a/Engine/Plugins/Runtime/USDCore/Source/USDUtilities/Public/USDShadeConversion.h
+++ b/Engine/Plugins/Runtime/USDCore/Source/USDUtilities/Public/USDShadeConversion.h
@@ -20,7 +20,9 @@
 #include "pxr/usd/usdShade/tokens.h"
 #include "USDIncludesEnd.h"
 
+#if WITH_EDITOR && PLATFORM_WINDOWS
 #include "MaterialXCore/Document.h"
+#endif
 
 PXR_NAMESPACE_OPEN_SCOPE
 	class UsdAttribute;
@@ -138,8 +140,9 @@ namespace UsdToUnreal
 	 * @param UsdShadeMaterialPrim - The prim used to generate the shader graph, should be a valid pxr::UsdShadeMaterial
 	 * @param RenderContext - Which render context output to read from the UsdShadeMaterial
 	 */
+#if WITH_EDITOR && PLATFORM_WINDOWS
 	USDUTILITIES_API MaterialX::DocumentPtr CreateSurfaceMaterial(const pxr::UsdPrim& UsdShadeMaterialPrim, const TCHAR* RenderContext);
-
+#endif
 	/**
 	 * Attemps to assign the values of the surface shader inputs to the MaterialInstance parameters by matching the inputs display names to the
 	 * parameters names.
diff --git a/Engine/Plugins/Runtime/USDCore/Source/USDUtilities/USDUtilities.Build.cs b/Engine/Plugins/Runtime/USDCore/Source/USDUtilities/USDUtilities.Build.cs
--- a/Engine/Plugins/Runtime/USDCore/Source/USDUtilities/USDUtilities.Build.cs
+++ b/Engine/Plugins/Runtime/USDCore/Source/USDUtilities/USDUtilities.Build.cs
@@ -17,7 +17,6 @@ namespace UnrealBuildTool.Rules
 					"Boost",
 					"Core",
 					"CoreUObject",
-					"MaterialX",
 					"UnrealUSDWrapper",
 					"USDClasses", // So that consumers can also include IUsdClassesModule for the new definition of FDisplayColorMaterial
 				}
@@ -64,6 +63,16 @@ namespace UnrealBuildTool.Rules
 			}
 
 			PrivateDefinitions.Add("SUPPRESS_PER_MODULE_INLINE_FILE"); // This module does not use core's standard operator new/delete overloads
+
+			if (Target.Platform == UnrealTargetPlatform.Win64)
+			{
+				PublicDependencyModuleNames.AddRange(
+					new string[]
+					{
+						"MaterialX"
+					}
+				);
+			}
 		}
 	}
 }
