USD: Ensure duplicated prims are unique across the entire stage.

diff --git a/Engine/Plugins/Runtime/USDCore/Source/USDUtilities/Private/USDConversionUtils.cpp b/Engine/Plugins/Runtime/USDCore/Source/USDUtilities/Private/USDConversionUtils.cpp
--- a/Engine/Plugins/Runtime/USDCore/Source/USDUtilities/Private/USDConversionUtils.cpp
+++ b/Engine/Plugins/Runtime/USDCore/Source/USDUtilities/Private/USDConversionUtils.cpp
@@ -3255,20 +3255,16 @@ TArray<UE::FSdfPath> UsdUtils::DuplicatePrims(const TArray<UE::FUsdPrim>& Prims,
 			const pxr::SdfPath ParentPath = UsdPrim.GetPath().GetParentPath();
 
 			int32 Suffix = -1;
-
-			bool bFoundName = false;
-			while (!bFoundName)
+			while (true)
 			{
 				NewSpecPath = ParentPath.AppendElementString(SourcePrimName + "_" + std::to_string(++Suffix));
 
-				bFoundName = true;
-				for (const pxr::SdfLayerHandle& Layer : LayersThatWillBeAffected)
+				// We want to make sure our new duplicate prim is unique across the entire composed stage, as opposed
+				// to silently overriding another prim that is only defined in an obscure layer somewhere
+				pxr::UsdPrim ExistingPrim = UsdStage->GetPrimAtPath(NewSpecPath);
+				if (!ExistingPrim)
 				{
-					if (Layer->HasSpec(NewSpecPath))
-					{
-						bFoundName = false;
-						break;
-					}
+					break;
 				}
 			}
 		}
