USD: Always produce skeletal import data / mesh descriptions with at least one UV set so that it can be used for static mesh builds too.

diff --git a/Engine/Plugins/Runtime/USDCore/Source/USDUtilities/Private/USDSkeletalDataConversion.cpp b/Engine/Plugins/Runtime/USDCore/Source/USDUtilities/Private/USDSkeletalDataConversion.cpp
--- a/Engine/Plugins/Runtime/USDCore/Source/USDUtilities/Private/USDSkeletalDataConversion.cpp
+++ b/Engine/Plugins/Runtime/USDCore/Source/USDUtilities/Private/USDSkeletalDataConversion.cpp
@@ -1750,6 +1750,18 @@ bool UsdToUnreal::ConvertSkinnedMesh(
 		++UVChannelIndex;
 	}
 
+	// Force our mesh data to always have at least one UV set.
+	// This so that we can have consistency across all our translated meshes, but it should only be needed in case
+	// Interchange wants to use the FMeshDescription produced from this SkelMeshImportData for a StaticMesh build,
+	// which assumes that at least one UV set is always available (Check the GetRawArray(0) within FStaticMeshOperations::ComputeMikktTangents,
+	// which was causing a crash for UE-224831)
+	if (UVSets.Num() == 0)
+	{
+		FUVSet& EmptySet = UVSets.Emplace_GetRef();
+		EmptySet.UVs.push_back(pxr::GfVec2f{0.0f, 0.0f});
+		EmptySet.InterpolationMethod = EUsdInterpolationMethod::Constant;
+	}
+
 	SkelMeshImportData.NumTexCoords = FMath::Clamp(FMath::Max(SkelMeshImportData.NumTexCoords, (uint32)UVSets.Num()), 0, MAX_TEXCOORDS);
 
 	SkelMeshImportData.Wedges.Reserve((NumExistingFaces + NumFaces) * 6);
