USD: Tweak wording on the warning toast show when trying to edit inside an instance.

diff --git a/Engine/Plugins/Runtime/USDCore/Source/USDUtilities/Private/USDAttributeUtils.cpp b/Engine/Plugins/Runtime/USDCore/Source/USDUtilities/Private/USDAttributeUtils.cpp
--- a/Engine/Plugins/Runtime/USDCore/Source/USDUtilities/Private/USDAttributeUtils.cpp
+++ b/Engine/Plugins/Runtime/USDCore/Source/USDUtilities/Private/USDAttributeUtils.cpp
@@ -408,12 +408,8 @@ bool UsdUtils::NotifyIfInstanceProxy(const pxr::UsdPrim& Prim)
 					const FText SubText = FText::Format(
 						LOCTEXT(
 							"AuthoringInsideInstanceSubText",
-							"Trying to author an opinion at or below prim:\n\n{0}\n\nThis prim is an instance proxy, and so cannot be "
-							"edited directly.\n\nIf you wish to modify just this prim, you can remove the instanceable flag for this prim and try "
-							"again. If you wish to "
-							"modify all instances, please edit the prims referenced by prim '{1}' directly."
+							"Trying to author an opinion below prim:\n\n{0}\n\nThis prim is an instance, so its child prim hierarchy cannot be modified directly.\n\nIf you wish to modify the hierarchy below just this particular instance, you can remove the instanceable flag from '{0}' and try again. If you wish to modify all instance hierarchies at the same time, please edit the prims referenced by prim '{0}' directly (i.e. open the referenced/payload layer directly, if any)."
 						),
-						FText::FromString(UsdToUnreal::ConvertPath(Prim.GetPrimPath())),
 						FText::FromString(FirstInstancePath)
 					);
 
