USD: Fix crash when modifying a collection while still using it.

diff --git a/Engine/Plugins/Runtime/USDCore/Source/UnrealUSDWrapper/Private/USDListener.cpp b/Engine/Plugins/Runtime/USDCore/Source/UnrealUSDWrapper/Private/USDListener.cpp
--- a/Engine/Plugins/Runtime/USDCore/Source/UnrealUSDWrapper/Private/USDListener.cpp
+++ b/Engine/Plugins/Runtime/USDCore/Source/UnrealUSDWrapper/Private/USDListener.cpp
@@ -511,9 +511,11 @@ namespace UE::USDListener::Private
 			{
 				if (Entry.Flags.bDidReloadContent)
 				{
-					OutInfoChanges = {};
 					OutResyncChanges = {};
 					OutResyncChanges.Add(RootPath, *RootInfoChange);
+
+					// Careful, as this will destroy our RootInfoChange/Entry
+					OutInfoChanges = {};
 					return true;
 				}
 			}
