USD: Prevent unnecessary patching of plugInfo.json files every editor launch.

diff --git a/Engine/Plugins/Runtime/USDCore/Source/USDClasses/Private/USDClassesModule.cpp b/Engine/Plugins/Runtime/USDCore/Source/USDClasses/Private/USDClassesModule.cpp
--- a/Engine/Plugins/Runtime/USDCore/Source/USDClasses/Private/USDClassesModule.cpp
+++ b/Engine/Plugins/Runtime/USDCore/Source/USDClasses/Private/USDClassesModule.cpp
@@ -199,12 +199,23 @@ void IUsdClassesModule::UpdatePlugInfoFiles(const FString& PluginDirectory, cons
 
 			// e.g. "../../../../../../Win64/sdf.dll"
 			FString NewLibraryPath = FPaths::Combine(PathToDll, DllName);
+
+			if (FPaths::IsSamePath(NewLibraryPath, LibraryPath))
+			{
+				// We don't expect anything else will modify these files, and we really need to only fix them up once
+				// per clean build/package/cook, so if we run into a path that already matches what we want we can just
+				// assume all the other plugInfo files are equally patched and early out.
+				//
+				// This is relevant because this code runs on FUnrealUSDWrapperModule::StartupModule, and so before
+				// *every single launch* of a project with the USD Importer enabled
+				UE_LOG(LogUsd, Log, TEXT("Found already updated LibraryPaths for USD plugInfo.json files at '%s'"), *PluginDirectory);
+				return;
+			}
+
 			(*PluginEntryObjectPtr)->SetField(TEXT("LibraryPath"), MakeShared<FJsonValueString>(NewLibraryPath));
-			bUpdatedFile = true;
 		}
 
-		// Serialize them back to JSON if we changed anything
-		if (bUpdatedFile)
+		// Serialize changes back to JSON
 		{
 			FString OutputString;
 			TSharedRef<TJsonWriter<>> Writer = TJsonWriterFactory<>::Create(&OutputString);
