USD: Fix crash when importing blend shapes with invalid point indices;

Fix exporter emitting incorrect blend shape indices when skeletal mesh had disabled sections.

diff --git a/Engine/Plugins/Runtime/USDCore/Source/USDUtilities/Private/USDSkeletalDataConversion.cpp b/Engine/Plugins/Runtime/USDCore/Source/USDUtilities/Private/USDSkeletalDataConversion.cpp
--- a/Engine/Plugins/Runtime/USDCore/Source/USDUtilities/Private/USDSkeletalDataConversion.cpp
+++ b/Engine/Plugins/Runtime/USDCore/Source/USDUtilities/Private/USDSkeletalDataConversion.cpp
@@ -618,7 +618,8 @@ namespace UnrealToUsdImpl
 		const TArray<FString>& MaterialAssignments,
 		const TArray<int32>& LODMaterialMap,
 		const pxr::UsdTimeCode TimeCode,
-		pxr::UsdPrim PrimToReceiveMaterialAssignments
+		pxr::UsdPrim PrimToReceiveMaterialAssignments,
+		TArray<int32>& OutSourceToPackedVertexIndex
 	)
 	{
 		FScopedUsdAllocs UsdAllocs;
@@ -642,7 +643,6 @@ namespace UnrealToUsdImpl
 		// FSkelMeshSection can be "disabled", at which point they don't show up in the engine. We'll skip those
 		// sections, and will use this array to help remap from a source vertex index to the vertex's corresponding
 		// index in a "packed" array of vertices that we'll push to USD
-		TArray<int32> SourceToPackedVertexIndex;
 		int32 PackedVertexIndex = 0;
 
 		// Vertices
@@ -656,7 +656,7 @@ namespace UnrealToUsdImpl
 			// vertices from disabled sections, which that function won't do
 			TArray<FSoftSkinVertex> Vertices;
 			Vertices.Reserve(LODModel.NumVertices);
-			SourceToPackedVertexIndex.SetNum(LODModel.NumVertices);
+			OutSourceToPackedVertexIndex.SetNum(LODModel.NumVertices);
 			for (int32 SectionIndex = 0; SectionIndex < LODModel.Sections.Num(); ++SectionIndex)
 			{
 				const FSkelMeshSection& Section = LODModel.Sections[SectionIndex];
@@ -665,14 +665,14 @@ namespace UnrealToUsdImpl
 					// Mark that these indices are skipped
 					for (int32 Index = 0; Index < Section.NumVertices; ++Index)
 					{
-						SourceToPackedVertexIndex[Section.BaseVertexIndex + Index] = INDEX_NONE;
+						OutSourceToPackedVertexIndex[Section.BaseVertexIndex + Index] = INDEX_NONE;
 					}
 					continue;
 				}
 
 				for (int32 Index = 0; Index < Section.NumVertices; ++Index)
 				{
-					SourceToPackedVertexIndex[Section.BaseVertexIndex + Index] = PackedVertexIndex++;
+					OutSourceToPackedVertexIndex[Section.BaseVertexIndex + Index] = PackedVertexIndex++;
 				}
 
 				Vertices.Append(Section.SoftVertices);
@@ -874,7 +874,7 @@ namespace UnrealToUsdImpl
 							{
 								const int32 SourceVertexIndex = LODModel.IndexBuffer[Section.BaseIndex + ((TriangleIndex * 3) + PointIndex)];
 
-								int32 PackedVertexPositionIndex = SourceToPackedVertexIndex[SourceVertexIndex];
+								int32 PackedVertexPositionIndex = OutSourceToPackedVertexIndex[SourceVertexIndex];
 
 								check(PackedVertexPositionIndex >= 0 && PackedVertexPositionIndex != INDEX_NONE);
 
@@ -1001,7 +1001,13 @@ namespace UnrealToUsdImpl
 	}
 
 	// Converts UE morph target deltas from DeltaArray into offsets, pointIndices and normalOffsets attributes of BlendShape
-	bool ConvertMorphTargetDeltas(const FMorphTargetDelta* DeltaArray, int32 NumDeltas, pxr::UsdSkelBlendShape& BlendShape, pxr::UsdTimeCode TimeCode)
+	bool ConvertMorphTargetDeltas(
+		const FMorphTargetDelta* DeltaArray,
+		int32 NumDeltas,
+		TArray<int32> SourceToPackedVertexIndex,
+		pxr::UsdSkelBlendShape& BlendShape,
+		pxr::UsdTimeCode TimeCode
+	)
 	{
 		if (!DeltaArray || NumDeltas == 0 || !BlendShape)
 		{
@@ -1023,8 +1029,14 @@ namespace UnrealToUsdImpl
 		{
 			const FMorphTargetDelta& Delta = DeltaArray[DeltaIndex];
 
+			int32 PackedIndex = SourceToPackedVertexIndex[Delta.SourceIdx];
+			if (PackedIndex == INDEX_NONE)
+			{
+				continue;
+			}
+
+			PointIndices.push_back(PackedIndex);
 			Offsets.push_back(UnrealToUsd::ConvertVectorFloat(StageInfo, (FVector)Delta.PositionDelta));
-			PointIndices.push_back(Delta.SourceIdx);
 			Normals.push_back(UnrealToUsd::ConvertVectorFloat(StageInfo, (FVector)Delta.TangentZDelta));
 		}
 
@@ -2878,6 +2890,11 @@ bool UsdToUnreal::ConvertSkeletalImportData(
 
 			for (const FMorphTargetDelta& Delta : BlendShape.Vertices)
 			{
+				if (!LODImportData.Points.IsValidIndex(Delta.SourceIdx))
+				{
+					continue;
+				}
+
 				NewModifiedPoints.Add(Delta.SourceIdx);
 				MorphTarget.Points.Add(Delta.PositionDelta + LODImportData.Points[Delta.SourceIdx]);
 			}
@@ -3713,6 +3730,7 @@ bool UnrealToUsd::ConvertSkeletalMesh(
 			LODMaterialMap = LODInfo->LODMaterialMap;
 		}
 
+		TArray<int32> SourceToPackedVertexIndex;
 		UnrealToUsdImpl::ConvertSkeletalMeshLOD(
 			SkeletalMesh,
 			LODModel,
@@ -3721,7 +3739,8 @@ bool UnrealToUsd::ConvertSkeletalMesh(
 			MaterialAssignments,
 			LODMaterialMap,
 			TimeCode,
-			MaterialPrim
+			MaterialPrim,
+			SourceToPackedVertexIndex
 		);
 
 		// Relationships can't target prims inside variants, so if we have BlendShapes to export we have to disable the edit target
@@ -3755,7 +3774,13 @@ bool UnrealToUsd::ConvertSkeletalMesh(
 			pxr::UsdPrim BlendShapePrim = UsdLODPrim.GetStage()->DefinePrim(BlendShapePath, UnrealToUsd::ConvertToken(TEXT("BlendShape")).Get());
 			pxr::UsdSkelBlendShape BlendShape{BlendShapePrim};
 
-			bool bCreatedBlendShape = UnrealToUsdImpl::ConvertMorphTargetDeltas(DeltaArray, NumDeltas, BlendShape, TimeCode);
+			bool bCreatedBlendShape = UnrealToUsdImpl::ConvertMorphTargetDeltas(
+				DeltaArray,
+				NumDeltas,
+				SourceToPackedVertexIndex,
+				BlendShape,
+				TimeCode
+			);
 			if (!bCreatedBlendShape)
 			{
 				continue;
