USD: Add nullptr prim check to NotifyIfInstanceProxy callback function.

diff --git a/Engine/Plugins/Runtime/USDCore/Source/USDUtilities/Private/USDAttributeUtils.cpp b/Engine/Plugins/Runtime/USDCore/Source/USDUtilities/Private/USDAttributeUtils.cpp
--- a/Engine/Plugins/Runtime/USDCore/Source/USDUtilities/Private/USDAttributeUtils.cpp
+++ b/Engine/Plugins/Runtime/USDCore/Source/USDUtilities/Private/USDAttributeUtils.cpp
@@ -316,6 +316,11 @@ bool UsdUtils::NotifyIfInstanceProxy(const pxr::UsdPrim& Prim)
 
 	TFunction<void()> RemoveInstanceables = [Prim = UE::FUsdPrim{Prim}]()
 	{
+		if (!Prim)
+		{
+			return;
+		}
+
 		UE::FUsdStage Stage = Prim.GetStage();
 		if (!Stage)
 		{
