USD: Change metersPerUnit handling to use double whenever possible.

diff --git a/Engine/Plugins/Importers/USDImporter/Source/USDStageImporter/Private/USDStageImportContext.cpp b/Engine/Plugins/Importers/USDImporter/Source/USDStageImporter/Private/USDStageImportContext.cpp
--- a/Engine/Plugins/Importers/USDImporter/Source/USDStageImporter/Private/USDStageImportContext.cpp
+++ b/Engine/Plugins/Importers/USDImporter/Source/USDStageImporter/Private/USDStageImportContext.cpp
@@ -112,7 +112,7 @@ void FUsdStageImportContext::Reset()
 	bReadFromStageCache = true;
 
 	bStageWasOriginallyOpenInCache = false;
-	OriginalMetersPerUnit = 0.01f;
+	OriginalMetersPerUnit = 0.01;
 	OriginalUpAxis = EUsdUpAxis::ZAxis;
 
 	bNeedsGarbageCollection = false;
diff --git a/Engine/Plugins/Importers/USDImporter/Source/USDStageImporter/Private/USDStageImportOptions.cpp b/Engine/Plugins/Importers/USDImporter/Source/USDStageImporter/Private/USDStageImportOptions.cpp
--- a/Engine/Plugins/Importers/USDImporter/Source/USDStageImporter/Private/USDStageImportOptions.cpp
+++ b/Engine/Plugins/Importers/USDImporter/Source/USDStageImporter/Private/USDStageImportOptions.cpp
@@ -38,7 +38,7 @@ UUsdStageImportOptions::UUsdStageImportOptions(const FObjectInitializer& ObjectI
 		false  /* bInvertFilters */
 	};
 	bOverrideStageOptions = false;
-	StageOptions.MetersPerUnit = 0.01f;
+	StageOptions.MetersPerUnit = 0.01;
 	StageOptions.UpAxis = EUsdUpAxis::ZAxis;
 	bImportAtSpecificTimeCode = false;
 	ImportTimeCode = 0.0f;
diff --git a/Engine/Plugins/Importers/USDImporter/Source/USDStageImporter/Public/USDStageImportContext.h b/Engine/Plugins/Importers/USDImporter/Source/USDStageImporter/Public/USDStageImportContext.h
--- a/Engine/Plugins/Importers/USDImporter/Source/USDStageImporter/Public/USDStageImportContext.h
+++ b/Engine/Plugins/Importers/USDImporter/Source/USDStageImporter/Public/USDStageImportContext.h
@@ -101,7 +101,7 @@ struct USDSTAGEIMPORTER_API FUsdStageImportContext
 
 	/** We modify the stage with our meters per unit import option on import. If the stage was already open, we use this to undo the changes after
 	 * import */
-	float OriginalMetersPerUnit;
+	double OriginalMetersPerUnit;
 	EUsdUpAxis OriginalUpAxis;
 
 	/** If we need to run GC after the import is complete */
diff --git a/Engine/Plugins/Interchange/Runtime/Source/Import/Private/Usd/InterchangeUsdTranslator.cpp b/Engine/Plugins/Interchange/Runtime/Source/Import/Private/Usd/InterchangeUsdTranslator.cpp
--- a/Engine/Plugins/Interchange/Runtime/Source/Import/Private/Usd/InterchangeUsdTranslator.cpp
+++ b/Engine/Plugins/Interchange/Runtime/Source/Import/Private/Usd/InterchangeUsdTranslator.cpp
@@ -2421,7 +2421,7 @@ UInterchangeUsdTranslatorSettings::UInterchangeUsdTranslatorSettings()
 	, InterpolationType(EUsdInterpolationType::Linear)
 	, bOverrideStageOptions(false)
 	, StageOptions{
-		  0.01f,			   // MetersPerUnit
+		  0.01,				   // MetersPerUnit
 		  EUsdUpAxis::ZAxis	   // UpAxis
 	  }
 {
diff --git a/Engine/Plugins/Runtime/USDCore/Source/USDClasses/Public/USDStageOptions.h b/Engine/Plugins/Runtime/USDCore/Source/USDClasses/Public/USDStageOptions.h
--- a/Engine/Plugins/Runtime/USDCore/Source/USDClasses/Public/USDStageOptions.h
+++ b/Engine/Plugins/Runtime/USDCore/Source/USDClasses/Public/USDStageOptions.h
@@ -20,7 +20,7 @@ struct USDCLASSES_API FUsdStageOptions
 
 	/** MetersPerUnit to use for the stage. Defaults to 0.01 (i.e. 1cm per unit, which equals UE units) */
 	UPROPERTY(EditAnywhere, config, BlueprintReadWrite, Category = "Stage options")
-	float MetersPerUnit = 0.01f;
+	double MetersPerUnit = 0.01;
 
 	/** UpAxis to use for the stage. Defaults to ZAxis, which equals the UE convention */
 	UPROPERTY(EditAnywhere, config, BlueprintReadWrite, Category = "Stage options")
diff --git a/Engine/Plugins/Runtime/USDCore/Source/USDUtilities/Private/USDConversionUtils.cpp b/Engine/Plugins/Runtime/USDCore/Source/USDUtilities/Private/USDConversionUtils.cpp
--- a/Engine/Plugins/Runtime/USDCore/Source/USDUtilities/Private/USDConversionUtils.cpp
+++ b/Engine/Plugins/Runtime/USDCore/Source/USDUtilities/Private/USDConversionUtils.cpp
@@ -391,19 +391,19 @@ void UsdUtils::SetUsdStageUpAxis(const pxr::UsdStageRefPtr& Stage, EUsdUpAxis Ax
 	SetUsdStageUpAxis(Stage, UpAxisToken);
 }
 
-float UsdUtils::GetUsdStageMetersPerUnit(const pxr::UsdStageRefPtr& Stage)
+double UsdUtils::GetUsdStageMetersPerUnit(const pxr::UsdStageRefPtr& Stage)
 {
-	return (float)pxr::UsdGeomGetStageMetersPerUnit(Stage);
+	return pxr::UsdGeomGetStageMetersPerUnit(Stage);
 }
 
-void UsdUtils::SetUsdStageMetersPerUnit(const pxr::UsdStageRefPtr& Stage, float MetersPerUnit)
+void UsdUtils::SetUsdStageMetersPerUnit(const pxr::UsdStageRefPtr& Stage, double MetersPerUnit)
 {
 	if (!Stage || !Stage->GetRootLayer())
 	{
 		return;
 	}
 
-	pxr::UsdEditContext(Stage, Stage->GetRootLayer());
+	pxr::UsdEditContext Context(Stage, Stage->GetRootLayer());
 	pxr::UsdGeomSetStageMetersPerUnit(Stage, MetersPerUnit);
 }
 
diff --git a/Engine/Plugins/Runtime/USDCore/Source/USDUtilities/Private/USDGroomConversion.cpp b/Engine/Plugins/Runtime/USDCore/Source/USDUtilities/Private/USDGroomConversion.cpp
--- a/Engine/Plugins/Runtime/USDCore/Source/USDUtilities/Private/USDGroomConversion.cpp
+++ b/Engine/Plugins/Runtime/USDCore/Source/USDUtilities/Private/USDGroomConversion.cpp
@@ -320,9 +320,10 @@ namespace UsdToUnreal
 				// Determine the scope of the Widths attribute
 				pxr::TfToken Interpolation = Curves.GetWidthsInterpolation();
 
-				const float UEMetersPerUnit = 0.01f;
-				const float Scale = !FMath::IsNearlyEqual(StageInfo.MetersPerUnit, UEMetersPerUnit) ? StageInfo.MetersPerUnit / UEMetersPerUnit
-																									: 1.0f;
+				const double UEMetersPerUnit = 0.01;
+				const float Scale = !FMath::IsNearlyEqual(StageInfo.MetersPerUnit, UEMetersPerUnit)
+										? (float)(StageInfo.MetersPerUnit / UEMetersPerUnit)
+										: 1.0f;
 				if (Interpolation == pxr::UsdGeomTokens->uniform || Interpolation == pxr::UsdGeomTokens->constant)
 				{
 					const float ConstWidth = WidthsArray[0] * Scale;
diff --git a/Engine/Plugins/Runtime/USDCore/Source/USDUtilities/Private/USDPrimConversion.cpp b/Engine/Plugins/Runtime/USDCore/Source/USDUtilities/Private/USDPrimConversion.cpp
--- a/Engine/Plugins/Runtime/USDCore/Source/USDUtilities/Private/USDPrimConversion.cpp
+++ b/Engine/Plugins/Runtime/USDCore/Source/USDUtilities/Private/USDPrimConversion.cpp
@@ -3323,7 +3323,7 @@ bool UnrealToUsd::ConvertInstancedStaticMeshComponent(const UInstancedStaticMesh
 		FVector Scale = USDTransform.GetScale3D();
 
 		// Compensate metersPerUnit
-		const float UEMetersPerUnit = 0.01f;
+		const double UEMetersPerUnit = 0.01;
 		if (!FMath::IsNearlyEqual(UEMetersPerUnit, StageInfo.MetersPerUnit))
 		{
 			Translation *= (UEMetersPerUnit / StageInfo.MetersPerUnit);
@@ -3872,7 +3872,7 @@ namespace UE::USDPrimConversion::Private
 			FVector Scale = USDTransform.GetScale3D();
 
 			// Compensate metersPerUnit
-			const float UEMetersPerUnit = 0.01f;
+			const double UEMetersPerUnit = 0.01;
 			if (!FMath::IsNearlyEqual(UEMetersPerUnit, StageInfo.MetersPerUnit))
 			{
 				Translation *= (UEMetersPerUnit / StageInfo.MetersPerUnit);
@@ -5007,7 +5007,7 @@ bool UnrealToUsd::CreateSkeletalAnimationBaker(
 				Swap(Translation.Y, Translation.Z);
 			}
 
-			Translations[BoneIndex] = UnrealToUsd::ConvertVectorFloat(Translation) * (0.01f / StageInfo.MetersPerUnit);
+			Translations[BoneIndex] = UnrealToUsd::ConvertVectorFloat(Translation) * ((float)(0.01 / StageInfo.MetersPerUnit));
 		}
 	};
 
diff --git a/Engine/Plugins/Runtime/USDCore/Source/USDUtilities/Private/USDTypesConversion.cpp b/Engine/Plugins/Runtime/USDCore/Source/USDUtilities/Private/USDTypesConversion.cpp
--- a/Engine/Plugins/Runtime/USDCore/Source/USDUtilities/Private/USDTypesConversion.cpp
+++ b/Engine/Plugins/Runtime/USDCore/Source/USDUtilities/Private/USDTypesConversion.cpp
@@ -100,7 +100,7 @@ namespace UsdUtils
 	{
 		FTransform TransformInUsdSpace = UsdUtils::ConvertAxes(StageInfo.UpAxis == EUsdUpAxis::ZAxis, TransformInUESpace);
 
-		const float UEMetersPerUnit = 0.01f;
+		const double UEMetersPerUnit = 0.01;
 		if (!FMath::IsNearlyEqual(StageInfo.MetersPerUnit, UEMetersPerUnit))
 		{
 			TransformInUsdSpace.ScaleTranslation(UEMetersPerUnit / StageInfo.MetersPerUnit);
@@ -113,7 +113,7 @@ namespace UsdUtils
 	{
 		FTransform TransformInUESpace = UsdUtils::ConvertAxes(StageInfo.UpAxis == EUsdUpAxis::ZAxis, TransformInUsdSpace);
 
-		const float UEMetersPerUnit = 0.01f;
+		const double UEMetersPerUnit = 0.01;
 		if (!FMath::IsNearlyEqual(StageInfo.MetersPerUnit, UEMetersPerUnit))
 		{
 			TransformInUESpace.ScaleTranslation(StageInfo.MetersPerUnit / UEMetersPerUnit);
@@ -210,7 +210,7 @@ namespace UsdToUnreal
 	{
 		FVector Value = ConvertVector(InValue);
 
-		const double UEMetersPerUnit = 0.01f;
+		const double UEMetersPerUnit = 0.01;
 		if (!FMath::IsNearlyEqual(StageInfo.MetersPerUnit, UEMetersPerUnit))
 		{
 			Value *= StageInfo.MetersPerUnit / UEMetersPerUnit;
@@ -301,7 +301,7 @@ namespace UsdToUnreal
 
 		Transform = UsdUtils::ConvertAxes(StageInfo.UpAxis == EUsdUpAxis::ZAxis, Transform);
 
-		const float UEMetersPerUnit = 0.01f;
+		const double UEMetersPerUnit = 0.01;
 		if (!FMath::IsNearlyEqual(StageInfo.MetersPerUnit, UEMetersPerUnit))
 		{
 			Transform.ScaleTranslation(StageInfo.MetersPerUnit / UEMetersPerUnit);
@@ -330,7 +330,7 @@ namespace UsdToUnreal
 
 	float ConvertDistance(const FUsdStageInfo& StageInfo, float Value)
 	{
-		const float UEMetersPerUnit = 0.01f;
+		const double UEMetersPerUnit = 0.01;
 		if (!FMath::IsNearlyEqual(StageInfo.MetersPerUnit, UEMetersPerUnit))
 		{
 			Value *= StageInfo.MetersPerUnit / UEMetersPerUnit;
@@ -433,7 +433,7 @@ namespace UnrealToUsd
 	template<typename T>
 	T ConvertVectorInner(const FUsdStageInfo& StageInfo, T Value)
 	{
-		const double UEMetersPerUnit = 0.01f;
+		const double UEMetersPerUnit = 0.01;
 		if (!FMath::IsNearlyEqual(StageInfo.MetersPerUnit, UEMetersPerUnit) && !FMath::IsNearlyZero(StageInfo.MetersPerUnit))
 		{
 			Value *= UEMetersPerUnit / StageInfo.MetersPerUnit;
@@ -567,7 +567,7 @@ namespace UnrealToUsd
 
 	float ConvertDistance(const FUsdStageInfo& StageInfo, float Value)
 	{
-		const float UEMetersPerUnit = 0.01f;
+		const double UEMetersPerUnit = 0.01;
 		if (!FMath::IsNearlyEqual(StageInfo.MetersPerUnit, UEMetersPerUnit))
 		{
 			Value *= UEMetersPerUnit / StageInfo.MetersPerUnit;
diff --git a/Engine/Plugins/Runtime/USDCore/Source/USDUtilities/Public/USDConversionUtils.h b/Engine/Plugins/Runtime/USDCore/Source/USDUtilities/Public/USDConversionUtils.h
--- a/Engine/Plugins/Runtime/USDCore/Source/USDUtilities/Public/USDConversionUtils.h
+++ b/Engine/Plugins/Runtime/USDCore/Source/USDUtilities/Public/USDConversionUtils.h
@@ -115,8 +115,8 @@ namespace UsdUtils
 	USDUTILITIES_API void SetUsdStageUpAxis(const pxr::UsdStageRefPtr& Stage, pxr::TfToken Axis);
 	USDUTILITIES_API void SetUsdStageUpAxis(const pxr::UsdStageRefPtr& Stage, EUsdUpAxis Axis);
 
-	USDUTILITIES_API float GetUsdStageMetersPerUnit(const pxr::UsdStageRefPtr& Stage);
-	USDUTILITIES_API void SetUsdStageMetersPerUnit(const pxr::UsdStageRefPtr& Stage, float MetersPerUnit);
+	USDUTILITIES_API double GetUsdStageMetersPerUnit(const pxr::UsdStageRefPtr& Stage);
+	USDUTILITIES_API void SetUsdStageMetersPerUnit(const pxr::UsdStageRefPtr& Stage, double MetersPerUnit);
 
 	USDUTILITIES_API int32 GetUsdStageNumFrames(const pxr::UsdStageRefPtr& Stage);
 
diff --git a/Engine/Plugins/Runtime/USDCore/Source/USDUtilities/Public/USDTypesConversion.h b/Engine/Plugins/Runtime/USDCore/Source/USDUtilities/Public/USDTypesConversion.h
--- a/Engine/Plugins/Runtime/USDCore/Source/USDUtilities/Public/USDTypesConversion.h
+++ b/Engine/Plugins/Runtime/USDCore/Source/USDUtilities/Public/USDTypesConversion.h
@@ -53,7 +53,7 @@ class TUsdStore;
 struct USDUTILITIES_API FUsdStageInfo
 {
 	EUsdUpAxis UpAxis = EUsdUpAxis::ZAxis;
-	float MetersPerUnit = 0.01f;
+	double MetersPerUnit = 0.01;
 
 	explicit FUsdStageInfo(const pxr::UsdStageRefPtr& Stage);
 };
