USD: Replace ANSI_TO_TCHAR and TCHAR_TO_ANSI with ANSI_TO_UTF8 and TCHAR_TO_UTF8.

diff --git a/Engine/Plugins/Importers/USDImporter/Source/USDStageEditorViewModels/Private/USDStageViewModel.cpp b/Engine/Plugins/Importers/USDImporter/Source/USDStageEditorViewModels/Private/USDStageViewModel.cpp
--- a/Engine/Plugins/Importers/USDImporter/Source/USDStageEditorViewModels/Private/USDStageViewModel.cpp
+++ b/Engine/Plugins/Importers/USDImporter/Source/USDStageEditorViewModels/Private/USDStageViewModel.cpp
@@ -247,7 +247,7 @@ void FUsdStageViewModel::SaveStageAs(const TCHAR* FilePath)
 
 			if (UE::FSdfLayer RootLayer = UsdStage.GetRootLayer())
 			{
-				if (pxr::SdfLayerRefPtr(RootLayer)->Export(TCHAR_TO_ANSI(FilePath)))
+				if (pxr::SdfLayerRefPtr(RootLayer)->Export(TCHAR_TO_UTF8(FilePath)))
 				{
 					FScopedUnrealAllocs UEAllocs;
 
diff --git a/Engine/Plugins/Interchange/Runtime/Source/Import/Private/MaterialX/MaterialXUtils/MaterialXSurfaceMaterial.cpp b/Engine/Plugins/Interchange/Runtime/Source/Import/Private/MaterialX/MaterialXUtils/MaterialXSurfaceMaterial.cpp
--- a/Engine/Plugins/Interchange/Runtime/Source/Import/Private/MaterialX/MaterialXUtils/MaterialXSurfaceMaterial.cpp
+++ b/Engine/Plugins/Interchange/Runtime/Source/Import/Private/MaterialX/MaterialXUtils/MaterialXSurfaceMaterial.cpp
@@ -27,7 +27,7 @@ void FMaterialXSurfaceMaterial::Translate(MaterialX::NodePtr SurfaceMaterialNode
 	 //We initialize the ShaderGraph outside of the loop, because a surfacematerial has only up to 2 inputs:
 	 //- a surfaceshader that we handle
 	 //- a displacementshader not yet supported
-	const FString ShaderGraphNodeUID = UInterchangeShaderNode::MakeNodeUid(ANSI_TO_TCHAR(SurfaceMaterialNode->getName().c_str()), FStringView{});
+	const FString ShaderGraphNodeUID = UInterchangeShaderNode::MakeNodeUid(UTF8_TO_TCHAR(SurfaceMaterialNode->getName().c_str()), FStringView{});
 	UInterchangeShaderGraphNode* ShaderGraphNode = const_cast<UInterchangeShaderGraphNode*>(Cast<UInterchangeShaderGraphNode>(NodeContainer.GetNode(ShaderGraphNodeUID)));
 
 	if(!ShaderGraphNode)
@@ -46,7 +46,7 @@ void FMaterialXSurfaceMaterial::Translate(MaterialX::NodePtr SurfaceMaterialNode
 		{
 			// we use here a static_cast instead of a dynamic cast, because we know for sure that the underlying type is a surfaceshader
 			// otherwise the document is ill-formed and the Translator already handles it
-			
+
 			TSharedPtr<FMaterialXSurfaceShaderAbstract> ShaderTranslator = StaticCastSharedPtr<FMaterialXSurfaceShaderAbstract>(FMaterialXManager::GetInstance().GetShaderTranslator(ConnectedNode->getCategory().c_str(), NodeContainer));
 			bHasSurfaceShader = ShaderTranslator.IsValid();
 			if(bHasSurfaceShader)
@@ -65,7 +65,7 @@ void FMaterialXSurfaceMaterial::Translate(MaterialX::NodePtr SurfaceMaterialNode
 	}
 	else
 	{
-		UE_LOG(LogInterchangeImport, Warning, TEXT("the surfaceshader of <%s> is not supported"), ANSI_TO_TCHAR(SurfaceMaterialNode->getName().c_str()));
+		UE_LOG(LogInterchangeImport, Warning, TEXT("the surfaceshader of <%s> is not supported"), UTF8_TO_TCHAR(SurfaceMaterialNode->getName().c_str()));
 	}
 }
 
diff --git a/Engine/Plugins/Interchange/Runtime/Source/Import/Private/MaterialX/MaterialXUtils/MaterialXSurfaceShaderAbstract.cpp b/Engine/Plugins/Interchange/Runtime/Source/Import/Private/MaterialX/MaterialXUtils/MaterialXSurfaceShaderAbstract.cpp
--- a/Engine/Plugins/Interchange/Runtime/Source/Import/Private/MaterialX/MaterialXUtils/MaterialXSurfaceShaderAbstract.cpp
+++ b/Engine/Plugins/Interchange/Runtime/Source/Import/Private/MaterialX/MaterialXUtils/MaterialXSurfaceShaderAbstract.cpp
@@ -1,4 +1,4 @@
-// Copyright Epic Games, Inc. All Rights Reserved. 
+// Copyright Epic Games, Inc. All Rights Reserved.
 
 #if WITH_EDITOR
 #include "MaterialX/MaterialXUtils/MaterialXSurfaceShaderAbstract.h"
@@ -363,7 +363,7 @@ void FMaterialXSurfaceShaderAbstract::ConnectNodeCategoryOutputToInput(const Mat
 		{
 			if(!ConnectMatchingNodeOutputToInput({ UpstreamNode, ParentShaderNode, InputChannelName, OutputChannelName }))
 			{
-				UE_LOG(LogInterchangeImport, Warning, TEXT("<%s> is not supported."), ANSI_TO_TCHAR(UpstreamNode->getCategory().c_str()));
+				UE_LOG(LogInterchangeImport, Warning, TEXT("<%s> is not supported."), UTF8_TO_TCHAR(UpstreamNode->getCategory().c_str()));
 			}
 		}
 	}
@@ -602,7 +602,7 @@ void FMaterialXSurfaceShaderAbstract::ConnectConvertInputToOutput(const FConnect
 	{
 		UInterchangeShaderNode* MaskShaderNode = nullptr;
 		if(NodeN == '3')
-		{			
+		{
 			MaskShaderNode = CreateMaskShaderNode(0b1110, Connect.UpstreamNode, Connect.UpstreamNode->getName().c_str());
 		}
 		else if(NodeN == '2')
@@ -618,7 +618,7 @@ void FMaterialXSurfaceShaderAbstract::ConnectConvertInputToOutput(const FConnect
 
 		UInterchangeShaderPortsAPI::ConnectDefaultOuputToInput(Connect.ParentShaderNode, Connect.InputChannelName, MaskShaderNode->GetUniqueID());
 	}
-	else // Append 
+	else // Append
 	{
 		//same as dot, just connect the next output to this parent input
 		SetAttributeNewName(Input, TCHAR_TO_UTF8(*Connect.InputChannelName)); //let's take the parent node's input name
@@ -951,7 +951,7 @@ void FMaterialXSurfaceShaderAbstract::ConnectNoise3DInputToOutput(const FConnect
 		// Connect the amplitude node to the shader node not the noise
 		// it will be handle during the upstream-downstream connection phase
 		// The index is here for the unicity of the attribute
-		Connect.UpstreamNode->setAttribute(mx::Attributes::ParentName + std::to_string(IndexAttrib), TCHAR_TO_ANSI(*ShaderNodeName));
+		Connect.UpstreamNode->setAttribute(mx::Attributes::ParentName + std::to_string(IndexAttrib), TCHAR_TO_UTF8(*ShaderNodeName));
 		AddAttributeFromValueOrInterface(Input, TEXT("B"), ShaderNode);
 
 		return ShaderNode;
@@ -1108,7 +1108,7 @@ void FMaterialXSurfaceShaderAbstract::ConnectBlurInputToOutput(const FConnectNod
 					}
 					else
 					{
-						UE_LOG(LogInterchangeImport, Warning, TEXT("<%s>: input 'size' must have a value."), ANSI_TO_TCHAR(Connect.UpstreamNode->getName().c_str()));
+						UE_LOG(LogInterchangeImport, Warning, TEXT("<%s>: input 'size' must have a value."), UTF8_TO_TCHAR(Connect.UpstreamNode->getName().c_str()));
 					}
 				}
 			}
@@ -1185,7 +1185,7 @@ void FMaterialXSurfaceShaderAbstract::ConnectSeparateInputToOutput(const FConnec
 	}
 	else
 	{
-		UE_LOG(LogInterchangeImport, Warning, TEXT("output <%s> not defined in <%s>."), *Connect.OutputName, ANSI_TO_TCHAR(Connect.UpstreamNode->getCategory().c_str()));
+		UE_LOG(LogInterchangeImport, Warning, TEXT("output <%s> not defined in <%s>."), *Connect.OutputName, UTF8_TO_TCHAR(Connect.UpstreamNode->getCategory().c_str()));
 	}
 
 	if(!bHasValueOrInterface)
@@ -1220,7 +1220,7 @@ void FMaterialXSurfaceShaderAbstract::ConnectSwitchInputToOutput(const FConnectN
 		}
 	}
 
-	//In that case we also need to add an attribute to Default	
+	//In that case we also need to add an attribute to Default
 	if (mx::InputPtr Input = Connect.UpstreamNode->getInput("in1"))
 	{
 		AddAttributeFromValueOrInterface(Input, Switch::Inputs::Default.ToString(), SwitchNode);
@@ -1228,7 +1228,7 @@ void FMaterialXSurfaceShaderAbstract::ConnectSwitchInputToOutput(const FConnectN
 		//Let's add a new input that is a copy of in1 to connect it to the Default input of the MaterialExpressionSwitch
 		mx::InputPtr InputDefault = Connect.UpstreamNode->addInput("default");
 		InputDefault->copyContentFrom(Input);
-		SetAttributeNewName(InputDefault, TCHAR_TO_ANSI(*Switch::Inputs::Default.ToString()));
+		SetAttributeNewName(InputDefault, TCHAR_TO_UTF8(*Switch::Inputs::Default.ToString()));
 	}
 
 	SwitchNode->AddInt32Attribute(Switch::Attributes::InputCount.ToString(), Index);
@@ -1370,7 +1370,7 @@ const FString& FMaterialXSurfaceShaderAbstract::GetMatchedInputName(MaterialX::N
 	if(Input)
 	{
 		const FString NodeCategory{ Node->getCategory().c_str() };
-		const FString InputName{ GetInputName(Input) };	
+		const FString InputName{ GetInputName(Input) };
 
 		if(const FString* Result = Manager.FindMatchingInput(NodeCategory, InputName, Node->getNodeDef(mx::EMPTY_STRING, true)->getNodeGroup().c_str()))
 		{
@@ -1481,7 +1481,7 @@ namespace
 
 			Buffer += Element->getName().c_str();
 			Buffer += TEXT("_");
-			Element->setAttribute(mx::Attributes::UniqueName, TCHAR_TO_ANSI(*FString{ Buffer.ToString() }.TrimChar(TEXT('_'))));
+			Element->setAttribute(mx::Attributes::UniqueName, TCHAR_TO_UTF8(*FString{ Buffer.ToString() }.TrimChar(TEXT('_'))));
 		}
 	}
 }
@@ -1565,7 +1565,7 @@ namespace
 	constexpr T DefaultMaxValue = std::numeric_limits<T>::max();
 
 	template<typename T>
-	std::string ValueToString = TCHAR_TO_ANSI(*FString::SanitizeFloat(DefaultMaxValue<T>));
+	std::string ValueToString = TCHAR_TO_UTF8(*FString::SanitizeFloat(DefaultMaxValue<T>));
 
 	template<typename T>
 	T GetValue(const auto& Value)
diff --git a/Engine/Plugins/Interchange/Runtime/Source/Import/Private/Usd/InterchangeUsdTranslator.cpp b/Engine/Plugins/Interchange/Runtime/Source/Import/Private/Usd/InterchangeUsdTranslator.cpp
--- a/Engine/Plugins/Interchange/Runtime/Source/Import/Private/Usd/InterchangeUsdTranslator.cpp
+++ b/Engine/Plugins/Interchange/Runtime/Source/Import/Private/Usd/InterchangeUsdTranslator.cpp
@@ -819,7 +819,7 @@ namespace UE::InterchangeUsdTranslator::Private
 			bool bHasShaderGraph = false;
 			if (MaterialX::DocumentPtr Document = UsdToUnreal::CreateSurfaceMaterial(Prim, TranslatorSettings ? *RenderContext.ToString() : nullptr))
 			{
-				if (MaterialX::NodePtr SurfaceMaterial = Document->getNode(TCHAR_TO_ANSI(*Prim.GetName().ToString())))
+				if (MaterialX::NodePtr SurfaceMaterial = Document->getNode(TCHAR_TO_UTF8(*Prim.GetName().ToString())))
 				{
 					try
 					{
@@ -835,7 +835,7 @@ namespace UE::InterchangeUsdTranslator::Private
 					}
 					catch (std::exception& Exception)
 					{
-						UE_LOG(LogUsd, Error, TEXT("%s"), ANSI_TO_TCHAR(Exception.what()));
+						UE_LOG(LogUsd, Error, TEXT("%s"), UTF8_TO_TCHAR(Exception.what()));
 					}
 				}
 
@@ -2511,9 +2511,9 @@ bool UInterchangeUSDTranslator::Translate(UInterchangeBaseNodeContainer& NodeCon
 	// allocate/deallocate on its own allocator.
 	MeshOptions.RenderContext = Settings->RenderContext == UnrealIdentifiers::UniversalRenderContext
 									? pxr::UsdShadeTokens->universalRenderContext
-									: pxr::TfToken{TCHAR_TO_ANSI(*Settings->RenderContext.ToString())};
+									: pxr::TfToken{TCHAR_TO_UTF8(*Settings->RenderContext.ToString())};
 	MeshOptions.MaterialPurpose = Settings->MaterialPurpose.IsNone() ? pxr::UsdShadeTokens->allPurpose
-																	 : pxr::TfToken{TCHAR_TO_ANSI(*Settings->MaterialPurpose.ToString())};
+																	 : pxr::TfToken{TCHAR_TO_UTF8(*Settings->MaterialPurpose.ToString())};
 
 	// Traverse stage and emit translated nodes
 	FTraversalInfo Info;
diff --git a/Engine/Plugins/Runtime/USDCore/Source/USDUtilities/Private/USDConversionUtils.cpp b/Engine/Plugins/Runtime/USDCore/Source/USDUtilities/Private/USDConversionUtils.cpp
--- a/Engine/Plugins/Runtime/USDCore/Source/USDUtilities/Private/USDConversionUtils.cpp
+++ b/Engine/Plugins/Runtime/USDCore/Source/USDUtilities/Private/USDConversionUtils.cpp
@@ -1819,13 +1819,13 @@ TArray<UE::FUsdPrim> UsdUtils::GetAllPrimsOfType(
 	TArray<UE::FUsdPrim> Result;
 
 #if USE_USD_SDK
-	const pxr::TfType SchemaType = pxr::TfType::FindByName(TCHAR_TO_ANSI(SchemaName));
+	const pxr::TfType SchemaType = pxr::TfType::FindByName(TCHAR_TO_UTF8(SchemaName));
 
 	TArray<TUsdStore<pxr::TfType>> ExcludeSchemaTypes;
 	ExcludeSchemaTypes.Reserve(ExcludeSchemaNames.Num());
 	for (const TCHAR* ExcludeSchemaName : ExcludeSchemaNames)
 	{
-		ExcludeSchemaTypes.Add(pxr::TfType(pxr::TfType::FindByName(TCHAR_TO_ANSI(ExcludeSchemaName))));
+		ExcludeSchemaTypes.Add(pxr::TfType(pxr::TfType::FindByName(TCHAR_TO_UTF8(ExcludeSchemaName))));
 	}
 
 	auto UsdPruneChildren = [&PruneChildren](const pxr::UsdPrim& ChildPrim) -> bool
diff --git a/Engine/Plugins/Runtime/USDCore/Source/USDUtilities/Private/USDShadeConversion.cpp b/Engine/Plugins/Runtime/USDCore/Source/USDUtilities/Private/USDShadeConversion.cpp
--- a/Engine/Plugins/Runtime/USDCore/Source/USDUtilities/Private/USDShadeConversion.cpp
+++ b/Engine/Plugins/Runtime/USDCore/Source/USDUtilities/Private/USDShadeConversion.cpp
@@ -2217,7 +2217,7 @@ namespace UE::UsdShadeConversion::Private
 				}
 				else
 				{
-					UE_LOG(LogUsd, Warning, TEXT("Couldn't find a supported type for (%s)."), ANSI_TO_TCHAR(InputType.c_str()));
+					UE_LOG(LogUsd, Warning, TEXT("Couldn't find a supported type for (%s)."), UTF8_TO_TCHAR(InputType.c_str()));
 				}
 			}
 
@@ -2238,7 +2238,7 @@ namespace UE::UsdShadeConversion::Private
 				if (!NodeDef)
 				{
 					Node = nullptr;
-					UE_LOG(LogUsd, Warning, TEXT("Couldn't find a nodedef for (%s)."), ANSI_TO_TCHAR(NodeDefTokenId.GetString().c_str()));
+					UE_LOG(LogUsd, Warning, TEXT("Couldn't find a nodedef for (%s)."), UTF8_TO_TCHAR(NodeDefTokenId.GetString().c_str()));
 					return false;
 				}
 
@@ -2361,11 +2361,11 @@ namespace UE::UsdShadeConversion::Private
 						Layer.Get()
 						? UsdToUnreal::ConvertString(TUsdStore<std::string>(pxr::SdfComputeAssetPathRelativeToLayer(Layer.Get(), AssetPath.GetAssetPath())).Get())
 						: AssetPath.GetAssetPath().c_str();
-					Input->setValue(std::string{ TCHAR_TO_ANSI(*PathToResolve) });
+					Input->setValue(std::string{ TCHAR_TO_UTF8(*PathToResolve) });
 				}
 				else
 				{
-					UE_LOG(LogUsd, Warning, TEXT("Couldn't find a value type for (%s)."), ANSI_TO_TCHAR(InputType.c_str()));
+					UE_LOG(LogUsd, Warning, TEXT("Couldn't find a value type for (%s)."), UTF8_TO_TCHAR(InputType.c_str()));
 				}
 			}
 		}
diff --git a/Engine/Plugins/Runtime/USDCore/Source/USDUtilities/Private/USDTypesConversion.cpp b/Engine/Plugins/Runtime/USDCore/Source/USDUtilities/Private/USDTypesConversion.cpp
--- a/Engine/Plugins/Runtime/USDCore/Source/USDUtilities/Private/USDTypesConversion.cpp
+++ b/Engine/Plugins/Runtime/USDCore/Source/USDUtilities/Private/USDTypesConversion.cpp
@@ -349,17 +349,17 @@ namespace UnrealToUsd
 
 	TUsdStore<pxr::SdfPath> ConvertPath(const TCHAR* InString)
 	{
-		return MakeUsdStore<pxr::SdfPath>(TCHAR_TO_ANSI(InString));
+		return MakeUsdStore<pxr::SdfPath>(TCHAR_TO_UTF8(InString));
 	}
 
 	TUsdStore<std::string> ConvertName(const FName& InName)
 	{
-		return MakeUsdStore<std::string>(TCHAR_TO_ANSI(*InName.ToString()));
+		return MakeUsdStore<std::string>(TCHAR_TO_UTF8(*InName.ToString()));
 	}
 
 	TUsdStore<pxr::TfToken> ConvertToken(const TCHAR* InString)
 	{
-		return MakeUsdStore<pxr::TfToken>(TCHAR_TO_ANSI(InString));
+		return MakeUsdStore<pxr::TfToken>(TCHAR_TO_UTF8(InString));
 	}
 
 	pxr::GfVec4f ConvertColor(const FLinearColor& InValue)
diff --git a/Engine/Plugins/Runtime/USDCore/Source/UnrealUSDWrapper/Private/USDListener.cpp b/Engine/Plugins/Runtime/USDCore/Source/UnrealUSDWrapper/Private/USDListener.cpp
--- a/Engine/Plugins/Runtime/USDCore/Source/UnrealUSDWrapper/Private/USDListener.cpp
+++ b/Engine/Plugins/Runtime/USDCore/Source/UnrealUSDWrapper/Private/USDListener.cpp
@@ -365,13 +365,13 @@ namespace UE::USDListener::Private
 		{
 			UsdUtils::FAttributeChange& ConvertedAttributeChange = OutConverted.AttributeChanges.Emplace_GetRef();
 
-			const FString FieldToken = ANSI_TO_TCHAR(AttributeChange.first.GetString().c_str());
+			const FString FieldToken = UTF8_TO_TCHAR(AttributeChange.first.GetString().c_str());
 
 			// For regular properties (most common case) FullFieldPath will already carry the property and FieldToken will be "default",
 			// "timeSamples", "variability", etc.
 			if (ChangedObject.IsPropertyPath())
 			{
-				ConvertedAttributeChange.PropertyName = ANSI_TO_TCHAR(ChangedObject.GetName().c_str());
+				ConvertedAttributeChange.PropertyName = UTF8_TO_TCHAR(ChangedObject.GetName().c_str());
 				ConvertedAttributeChange.Field = FieldToken;
 			}
 			// For stage properties it seems like USD likes to send just "/" as the FullFieldPath, and the actual property name in FieldToken (e.g.
@@ -397,7 +397,7 @@ namespace UE::USDListener::Private
 
 			if (ChangedObject.IsPropertyPath())
 			{
-				ConvertedAttributeChange.PropertyName = ANSI_TO_TCHAR(ChangedObject.GetName().c_str());
+				ConvertedAttributeChange.PropertyName = UTF8_TO_TCHAR(ChangedObject.GetName().c_str());
 			}
 		}
 
@@ -445,13 +445,13 @@ namespace UE::USDListener::Private
 		for (const std::pair<std::string, pxr::SdfChangeList::SubLayerChangeType>& SubLayerChange : Entry.subLayerChanges)
 		{
 			OutConverted.SubLayerChanges.Add(TPair<FString, UsdUtils::ESubLayerChangeType>(
-				ANSI_TO_TCHAR(SubLayerChange.first.c_str()),
+				UTF8_TO_TCHAR(SubLayerChange.first.c_str()),
 				static_cast<UsdUtils::ESubLayerChangeType>(SubLayerChange.second)
 			));
 		}
 
-		OutConverted.OldPath = ANSI_TO_TCHAR(Entry.oldPath.GetString().c_str());
-		OutConverted.OldIdentifier = ANSI_TO_TCHAR(Entry.oldIdentifier.c_str());
+		OutConverted.OldPath = UTF8_TO_TCHAR(Entry.oldPath.GetString().c_str());
+		OutConverted.OldIdentifier = UTF8_TO_TCHAR(Entry.oldIdentifier.c_str());
 	}
 
 	bool ConvertPathRange(const pxr::UsdNotice::ObjectsChanged::PathRange& PathRange, UsdUtils::FObjectChangesByPath& OutChanges)
@@ -464,7 +464,7 @@ namespace UE::USDListener::Private
 
 		for (UsdNotice::ObjectsChanged::PathRange::const_iterator It = PathRange.begin(); It != PathRange.end(); ++It)
 		{
-			const FString PrimPath = ANSI_TO_TCHAR(It->GetAbsoluteRootOrPrimPath().GetAsString().c_str());
+			const FString PrimPath = UTF8_TO_TCHAR(It->GetAbsoluteRootOrPrimPath().GetAsString().c_str());
 
 			if (pxr::UsdPrim::IsPathInPrototype(UE::FSdfPath(*PrimPath)))
 			{
@@ -472,7 +472,7 @@ namespace UE::USDListener::Private
 			}
 
 			// Something like "/Root/Prim.some_field", or "/"
-			const FString FullFieldPath = ANSI_TO_TCHAR(It->GetAsString().c_str());
+			const FString FullFieldPath = UTF8_TO_TCHAR(It->GetAsString().c_str());
 
 			TArray<UsdUtils::FSdfChangeListEntry>& ConvertedChanges = OutChanges.FindOrAdd(PrimPath);
 
@@ -549,7 +549,7 @@ namespace UE::USDListener::Private
 				{
 					for (UsdUtils::FAttributeChange& Change : ChangeIt->AttributeChanges)
 					{
-						if (Change.PropertyName == ANSI_TO_TCHAR(pxr::UsdGeomTokens->visibility.GetString().c_str()))
+						if (Change.PropertyName == UTF8_TO_TCHAR(pxr::UsdGeomTokens->visibility.GetString().c_str()))
 						{
 							bUpgrade = true;
 							break;
@@ -835,7 +835,7 @@ void FUsdListenerImpl::HandleLayersChangedNotice(const pxr::SdfNotice::LayersDid
 				const pxr::SdfChangeList::Entry::_Flags& Flags = Change.second.flags;
 				if (Flags.didReloadContent)
 				{
-					LayersNames.Add(ANSI_TO_TCHAR(ChangeVecItem.first->GetIdentifier().c_str()));
+					LayersNames.Add(UTF8_TO_TCHAR(ChangeVecItem.first->GetIdentifier().c_str()));
 				}
 			}
 		}
diff --git a/Engine/Plugins/Runtime/USDCore/Source/UnrealUSDWrapper/Private/UnrealUSDWrapper.cpp b/Engine/Plugins/Runtime/USDCore/Source/UnrealUSDWrapper/Private/UnrealUSDWrapper.cpp
--- a/Engine/Plugins/Runtime/USDCore/Source/UnrealUSDWrapper/Private/UnrealUSDWrapper.cpp
+++ b/Engine/Plugins/Runtime/USDCore/Source/UnrealUSDWrapper/Private/UnrealUSDWrapper.cpp
@@ -455,15 +455,15 @@ FName UnrealIdentifiers::LightSourceAnglePropertyName = GET_MEMBER_NAME_CHECKED(
 
 FString UnrealIdentifiers::MaterialAllPurposeText = TEXT("allPurpose");
 #if USE_USD_SDK
-FString UnrealIdentifiers::MaterialAllPurpose = ANSI_TO_TCHAR(pxr::UsdShadeTokens->allPurpose.GetString().c_str());
-FString UnrealIdentifiers::MaterialPreviewPurpose = ANSI_TO_TCHAR(pxr::UsdShadeTokens->preview.GetString().c_str());
-FString UnrealIdentifiers::MaterialFullPurpose = ANSI_TO_TCHAR(pxr::UsdShadeTokens->full.GetString().c_str());
-FString UnrealIdentifiers::PrimvarsDisplayColor = ANSI_TO_TCHAR(pxr::UsdGeomTokens->primvarsDisplayColor.GetString().c_str());
-FString UnrealIdentifiers::PrimvarsDisplayOpacity = ANSI_TO_TCHAR(pxr::UsdGeomTokens->primvarsDisplayOpacity.GetString().c_str());
-FString UnrealIdentifiers::DoubleSided = ANSI_TO_TCHAR(pxr::UsdGeomTokens->doubleSided.GetString().c_str());
-FString UnrealIdentifiers::ModelDrawMode = ANSI_TO_TCHAR(pxr::UsdGeomTokens->modelDrawMode.GetString().c_str());
-FString UnrealIdentifiers::ModelApplyDrawMode = ANSI_TO_TCHAR(pxr::UsdGeomTokens->modelApplyDrawMode.GetString().c_str());
-FString UnrealIdentifiers::UsdNamespaceDelimiter = ANSI_TO_TCHAR(pxr::SdfPathTokens->namespaceDelimiter.GetString().c_str());
+FString UnrealIdentifiers::MaterialAllPurpose = UTF8_TO_TCHAR(pxr::UsdShadeTokens->allPurpose.GetString().c_str());
+FString UnrealIdentifiers::MaterialPreviewPurpose = UTF8_TO_TCHAR(pxr::UsdShadeTokens->preview.GetString().c_str());
+FString UnrealIdentifiers::MaterialFullPurpose = UTF8_TO_TCHAR(pxr::UsdShadeTokens->full.GetString().c_str());
+FString UnrealIdentifiers::PrimvarsDisplayColor = UTF8_TO_TCHAR(pxr::UsdGeomTokens->primvarsDisplayColor.GetString().c_str());
+FString UnrealIdentifiers::PrimvarsDisplayOpacity = UTF8_TO_TCHAR(pxr::UsdGeomTokens->primvarsDisplayOpacity.GetString().c_str());
+FString UnrealIdentifiers::DoubleSided = UTF8_TO_TCHAR(pxr::UsdGeomTokens->doubleSided.GetString().c_str());
+FString UnrealIdentifiers::ModelDrawMode = UTF8_TO_TCHAR(pxr::UsdGeomTokens->modelDrawMode.GetString().c_str());
+FString UnrealIdentifiers::ModelApplyDrawMode = UTF8_TO_TCHAR(pxr::UsdGeomTokens->modelApplyDrawMode.GetString().c_str());
+FString UnrealIdentifiers::UsdNamespaceDelimiter = UTF8_TO_TCHAR(pxr::SdfPathTokens->namespaceDelimiter.GetString().c_str());
 #else
 FString UnrealIdentifiers::MaterialAllPurpose = TEXT("");
 FString UnrealIdentifiers::MaterialPreviewPurpose = TEXT("preview");
@@ -479,7 +479,7 @@ FString UnrealIdentifiers::UsdNamespaceDelimiter = TEXT(":");
 #if USE_USD_SDK
 // This is largely for show, as the universalRenderContext is an empty string, and FNames created from empty strings equal NAME_None anyway
 // i.e.: NAME_None == FName(TEXT("")) == FName(TEXT("None")) == UnrealIdentifiers::UniversalRenderContext
-const FName UnrealIdentifiers::UniversalRenderContext = ANSI_TO_TCHAR(pxr::UsdShadeTokens->universalRenderContext.GetString().c_str());
+const FName UnrealIdentifiers::UniversalRenderContext = UTF8_TO_TCHAR(pxr::UsdShadeTokens->universalRenderContext.GetString().c_str());
 #else
 const FName UnrealIdentifiers::UniversalRenderContext = NAME_None;
 #endif
@@ -523,12 +523,12 @@ public:
 		Msg += ": ";
 		Msg += Error.GetCommentary();
 
-		UE_LOG(LogUsd, Error, TEXT("%s"), ANSI_TO_TCHAR(Msg.c_str()));
+		UE_LOG(LogUsd, Error, TEXT("%s"), UTF8_TO_TCHAR(Msg.c_str()));
 	}
 
 	virtual void IssueFatalError(const pxr::TfCallContext& Context, const std::string& Msg) override
 	{
-		UE_LOG(LogUsd, Error, TEXT("%s"), ANSI_TO_TCHAR(Msg.c_str()));
+		UE_LOG(LogUsd, Error, TEXT("%s"), UTF8_TO_TCHAR(Msg.c_str()));
 	}
 
 	virtual void IssueStatus(const pxr::TfStatus& Status) override
@@ -539,7 +539,7 @@ public:
 		Msg += ": ";
 		Msg += Status.GetCommentary();
 
-		UE_LOG(LogUsd, Log, TEXT("%s"), ANSI_TO_TCHAR(Msg.c_str()));
+		UE_LOG(LogUsd, Log, TEXT("%s"), UTF8_TO_TCHAR(Msg.c_str()));
 	}
 
 	virtual void IssueWarning(const pxr::TfWarning& Warning) override
@@ -550,7 +550,7 @@ public:
 		Msg += ": ";
 		Msg += Warning.GetCommentary();
 
-		UE_LOG(LogUsd, Warning, TEXT("%s"), ANSI_TO_TCHAR(Msg.c_str()));
+		UE_LOG(LogUsd, Warning, TEXT("%s"), UTF8_TO_TCHAR(Msg.c_str()));
 	}
 };
 #else
@@ -647,7 +647,7 @@ TArray<FString> UnrealUSDWrapper::GetAllSupportedFileFormats()
 			continue;
 		}
 
-		Result.Emplace(ANSI_TO_TCHAR(Ext.c_str()));
+		Result.Emplace(UTF8_TO_TCHAR(Ext.c_str()));
 	}
 #endif	  // #if USE_USD_SDK
 
@@ -687,7 +687,7 @@ TArray<FString> UnrealUSDWrapper::GetNativeFileFormats()
 
 	for (const std::string& FileExtension : FileExtensions)
 	{
-		Result.Emplace(ANSI_TO_TCHAR(FileExtension.c_str()));
+		Result.Emplace(UTF8_TO_TCHAR(FileExtension.c_str()));
 	}
 #endif	  // #if USE_USD_SDK
 
@@ -714,7 +714,7 @@ void UnrealUSDWrapper::GetNativeFileFormats(TArray<FString>& OutTextFormats, TAr
 	OutTextFormats.Reset(TextExtensions.size());
 	for (const std::string& Extension : TextExtensions)
 	{
-		OutTextFormats.Emplace(ANSI_TO_TCHAR(Extension.c_str()));
+		OutTextFormats.Emplace(UTF8_TO_TCHAR(Extension.c_str()));
 	}
 
 	const static std::vector<pxr::TfToken> BinaryFormatIds{
@@ -740,7 +740,7 @@ void UnrealUSDWrapper::GetNativeFileFormats(TArray<FString>& OutTextFormats, TAr
 	OutPossiblyBinaryFormats.Reset(PossiblyBinaryExtensions.size());
 	for (const std::string& Extension : PossiblyBinaryExtensions)
 	{
-		OutPossiblyBinaryFormats.Emplace(ANSI_TO_TCHAR(Extension.c_str()));
+		OutPossiblyBinaryFormats.Emplace(UTF8_TO_TCHAR(Extension.c_str()));
 	}
 #endif	  // #if USE_USD_SDK
 }
@@ -832,7 +832,7 @@ namespace UE::UnrealUSDWrapper::Private
 
 			for (const FString& AllowedPrimPath : *PopulationMask)
 			{
-				Mask.Add(pxr::SdfPath{TCHAR_TO_ANSI(*AllowedPrimPath)});
+				Mask.Add(pxr::SdfPath{TCHAR_TO_UTF8(*AllowedPrimPath)});
 			}
 		}
 
@@ -846,8 +846,8 @@ namespace UE::UnrealUSDWrapper::Private
 		RootIdentifierStr.RemoveFromStart(UnrealIdentifiers::IdentifierPrefix);
 		SessionIdentifierStr.RemoveFromStart(UnrealIdentifiers::IdentifierPrefix);
 
-		pxr::SdfLayerRefPtr RootLayer = pxr::SdfLayer::FindOrOpen(TCHAR_TO_ANSI(*RootIdentifierStr));
-		pxr::SdfLayerRefPtr SessionLayer = pxr::SdfLayer::FindOrOpen(TCHAR_TO_ANSI(*SessionIdentifierStr));
+		pxr::SdfLayerRefPtr RootLayer = pxr::SdfLayer::FindOrOpen(TCHAR_TO_UTF8(*RootIdentifierStr));
+		pxr::SdfLayerRefPtr SessionLayer = pxr::SdfLayer::FindOrOpen(TCHAR_TO_UTF8(*SessionIdentifierStr));
 		if (RootLayer)
 		{
 			if (PopulationMask)
@@ -993,7 +993,7 @@ UE::FUsdStage UnrealUSDWrapper::NewStage(const TCHAR* FilePath)
 #if USE_USD_SDK
 	FScopedUsdAllocs UsdAllocs;
 
-	UE::FUsdStage UsdStage(pxr::UsdStage::CreateNew(TCHAR_TO_ANSI(FilePath)));
+	UE::FUsdStage UsdStage(pxr::UsdStage::CreateNew(TCHAR_TO_UTF8(FilePath)));
 	if (UsdStage)
 	{
 		pxr::UsdGeomSetStageUpAxis(UsdStage, pxr::UsdGeomTokens->z);
@@ -1035,7 +1035,7 @@ UE::FUsdStage UnrealUSDWrapper::GetClipboardStage()
 	{
 		// Use this specific name so that we can check for it on our automated tests
 		pxr::UsdStageCacheContext StageCacheContext{pxr::UsdUtilsStageCache::Get()};
-		ClipboardStage = UE::FUsdStage{pxr::UsdStage::CreateInMemory(TCHAR_TO_ANSI(TEXT("unreal_clipboard_layer")))};
+		ClipboardStage = UE::FUsdStage{pxr::UsdStage::CreateInMemory(TCHAR_TO_UTF8(TEXT("unreal_clipboard_layer")))};
 
 		pxr::UsdGeomSetStageUpAxis(ClipboardStage, pxr::UsdGeomTokens->z);
 	}
diff --git a/Engine/Plugins/Runtime/USDCore/Source/UnrealUSDWrapper/Private/UsdWrappers/SdfAttributeSpec.cpp b/Engine/Plugins/Runtime/USDCore/Source/UnrealUSDWrapper/Private/UsdWrappers/SdfAttributeSpec.cpp
--- a/Engine/Plugins/Runtime/USDCore/Source/UnrealUSDWrapper/Private/UsdWrappers/SdfAttributeSpec.cpp
+++ b/Engine/Plugins/Runtime/USDCore/Source/UnrealUSDWrapper/Private/UsdWrappers/SdfAttributeSpec.cpp
@@ -164,7 +164,7 @@ namespace UE
 	{
 #if USE_USD_SDK
 		FScopedUsdAllocs UsdAllocs;
-		return ANSI_TO_TCHAR(Impl->PxrSdfAttributeSpec.Get()->GetName().c_str());
+		return UTF8_TO_TCHAR(Impl->PxrSdfAttributeSpec.Get()->GetName().c_str());
 #else
 		return {};
 #endif	  // #if USE_USD_SDK
diff --git a/Engine/Plugins/Runtime/USDCore/Source/UnrealUSDWrapper/Private/UsdWrappers/SdfLayer.cpp b/Engine/Plugins/Runtime/USDCore/Source/UnrealUSDWrapper/Private/UsdWrappers/SdfLayer.cpp
--- a/Engine/Plugins/Runtime/USDCore/Source/UnrealUSDWrapper/Private/UsdWrappers/SdfLayer.cpp
+++ b/Engine/Plugins/Runtime/USDCore/Source/UnrealUSDWrapper/Private/UsdWrappers/SdfLayer.cpp
@@ -261,7 +261,7 @@ namespace UE
 		if (const PtrType& Ptr = Impl->GetInner())
 		{
 			FScopedUsdAllocs UsdAllocs;
-			return ANSI_TO_TCHAR(Ptr->GetComment().c_str());
+			return UTF8_TO_TCHAR(Ptr->GetComment().c_str());
 		}
 #endif	  // #if USE_USD_SDK
 		return {};
@@ -274,7 +274,7 @@ namespace UE
 		if (const PtrType& Ptr = Impl->GetInner())
 		{
 			FScopedUsdAllocs UsdAllocs;
-			Ptr->SetComment(TCHAR_TO_ANSI(Comment));
+			Ptr->SetComment(TCHAR_TO_UTF8(Comment));
 		}
 #endif	  // #if USE_USD_SDK
 	}
@@ -317,10 +317,10 @@ namespace UE
 		std::map<std::string, std::string> UsdFileFormatArguments;
 		for (const TPair<FString, FString>& Pair : FileFormatArguments)
 		{
-			UsdFileFormatArguments.insert({TCHAR_TO_ANSI(*Pair.Key), TCHAR_TO_ANSI(*Pair.Value)});
+			UsdFileFormatArguments.insert({TCHAR_TO_UTF8(*Pair.Key), TCHAR_TO_UTF8(*Pair.Value)});
 		}
 
-		return FSdfLayer(pxr::SdfLayer::FindOrOpen(TCHAR_TO_ANSI(Identifier), UsdFileFormatArguments));
+		return FSdfLayer(pxr::SdfLayer::FindOrOpen(TCHAR_TO_UTF8(Identifier), UsdFileFormatArguments));
 #else
 		return FSdfLayer();
 #endif	  // #if USE_USD_SDK
@@ -335,10 +335,10 @@ namespace UE
 		std::map<std::string, std::string> UsdFileFormatArguments;
 		for (const TPair<FString, FString>& Pair : FileFormatArguments)
 		{
-			UsdFileFormatArguments.insert({TCHAR_TO_ANSI(*Pair.Key), TCHAR_TO_ANSI(*Pair.Value)});
+			UsdFileFormatArguments.insert({TCHAR_TO_UTF8(*Pair.Key), TCHAR_TO_UTF8(*Pair.Value)});
 		}
 
-		return FSdfLayer(pxr::SdfLayer::CreateNew(TCHAR_TO_ANSI(Identifier), UsdFileFormatArguments));
+		return FSdfLayer(pxr::SdfLayer::CreateNew(TCHAR_TO_UTF8(Identifier), UsdFileFormatArguments));
 #else
 		return FSdfLayer();
 #endif	  // #if USE_USD_SDK
@@ -358,7 +358,7 @@ namespace UE
 			for (std::map<std::string, std::string>::const_iterator Iter = UsdFileFormatArguments.cbegin(); Iter != UsdFileFormatArguments.cend();
 				 ++Iter)
 			{
-				FileFormatArguments.Add(ANSI_TO_TCHAR(Iter->first.c_str()), ANSI_TO_TCHAR(Iter->second.c_str()));
+				FileFormatArguments.Add(UTF8_TO_TCHAR(Iter->first.c_str()), UTF8_TO_TCHAR(Iter->second.c_str()));
 			}
 		}
 #endif	  // #if USE_USD_SDK
@@ -392,7 +392,7 @@ namespace UE
 			Results.Reserve(AssetDependencies.size());
 			for (const std::string& AssetDependency : AssetDependencies)
 			{
-				Results.Add(ANSI_TO_TCHAR(AssetDependency.c_str()));
+				Results.Add(UTF8_TO_TCHAR(AssetDependency.c_str()));
 			}
 		}
 #endif	  // #if USE_USD_SDK
@@ -407,8 +407,8 @@ namespace UE
 		{
 			FScopedUsdAllocs UsdAllocs;
 			return Ptr->UpdateCompositionAssetDependency(
-				TCHAR_TO_ANSI(OldAssetPath),
-				(NewAssetPath != nullptr) ? TCHAR_TO_ANSI(NewAssetPath) : std::string()
+				TCHAR_TO_UTF8(OldAssetPath),
+				(NewAssetPath != nullptr) ? TCHAR_TO_UTF8(NewAssetPath) : std::string()
 			);
 		}
 #endif	  // #if USE_USD_SDK
@@ -422,7 +422,7 @@ namespace UE
 		if (const PtrType& Ptr = Impl->GetInner())
 		{
 			FScopedUsdAllocs UsdAllocs;
-			return FString(ANSI_TO_TCHAR(Ptr->GetRealPath().c_str()));
+			return FString(UTF8_TO_TCHAR(Ptr->GetRealPath().c_str()));
 		}
 #endif	  // #if USE_USD_SDK
 
@@ -436,7 +436,7 @@ namespace UE
 		if (const PtrType& Ptr = Impl->GetInner())
 		{
 			FScopedUsdAllocs UsdAllocs;
-			return FString(ANSI_TO_TCHAR(Ptr->GetIdentifier().c_str()));
+			return FString(UTF8_TO_TCHAR(Ptr->GetIdentifier().c_str()));
 		}
 #endif	  // #if USE_USD_SDK
 
@@ -450,7 +450,7 @@ namespace UE
 		if (const PtrType& Ptr = Impl->GetInner())
 		{
 			FScopedUsdAllocs UsdAllocs;
-			return FString(ANSI_TO_TCHAR(Ptr->GetDisplayName().c_str()));
+			return FString(UTF8_TO_TCHAR(Ptr->GetDisplayName().c_str()));
 		}
 #endif	  // #if USE_USD_SDK
 
@@ -464,7 +464,7 @@ namespace UE
 		if (const PtrType& Ptr = Impl->GetInner())
 		{
 			FScopedUsdAllocs UsdAllocs;
-			return FString(ANSI_TO_TCHAR(Ptr->ComputeAbsolutePath(TCHAR_TO_ANSI(*AssetPath)).c_str()));
+			return FString(UTF8_TO_TCHAR(Ptr->ComputeAbsolutePath(TCHAR_TO_UTF8(*AssetPath)).c_str()));
 		}
 #endif	  // #if USE_USD_SDK
 
@@ -890,10 +890,10 @@ namespace UE
 			std::map<std::string, std::string> UsdFileFormatArguments;
 			for (const TPair<FString, FString>& Pair : FileFormatArguments)
 			{
-				UsdFileFormatArguments.insert({TCHAR_TO_ANSI(*Pair.Key), TCHAR_TO_ANSI(*Pair.Value)});
+				UsdFileFormatArguments.insert({TCHAR_TO_UTF8(*Pair.Key), TCHAR_TO_UTF8(*Pair.Value)});
 			}
 
-			return Ptr->Export(TCHAR_TO_ANSI(Filename), TCHAR_TO_ANSI(*Comment), UsdFileFormatArguments);
+			return Ptr->Export(TCHAR_TO_UTF8(Filename), TCHAR_TO_UTF8(*Comment), UsdFileFormatArguments);
 		}
 #endif	  // #if USE_USD_SDK
 
@@ -926,7 +926,7 @@ namespace UE
 	{
 #if USE_USD_SDK
 		FScopedUsdAllocs UsdAllocs;
-		return FString(ANSI_TO_TCHAR(pxr::SdfComputeAssetPathRelativeToLayer(pxr::SdfLayerRefPtr(Anchor), TCHAR_TO_ANSI(AssetPath)).c_str()));
+		return FString(UTF8_TO_TCHAR(pxr::SdfComputeAssetPathRelativeToLayer(pxr::SdfLayerRefPtr(Anchor), TCHAR_TO_UTF8(AssetPath)).c_str()));
 #else
 		return FString();
 #endif	  // #if USE_USD_SDK
diff --git a/Engine/Plugins/Runtime/USDCore/Source/UnrealUSDWrapper/Private/UsdWrappers/SdfPath.cpp b/Engine/Plugins/Runtime/USDCore/Source/UnrealUSDWrapper/Private/UsdWrappers/SdfPath.cpp
--- a/Engine/Plugins/Runtime/USDCore/Source/UnrealUSDWrapper/Private/UsdWrappers/SdfPath.cpp
+++ b/Engine/Plugins/Runtime/USDCore/Source/UnrealUSDWrapper/Private/UsdWrappers/SdfPath.cpp
@@ -28,7 +28,7 @@ namespace UE
 				: PxrSdfPath(InSdfPath)
 			{
 #if ENABLE_USD_DEBUG_PATH
-				DebugPath = FString(ANSI_TO_TCHAR(PxrSdfPath.Get().GetString().c_str()));
+				DebugPath = FString(UTF8_TO_TCHAR(PxrSdfPath.Get().GetString().c_str()));
 #endif
 			}
 
@@ -36,7 +36,7 @@ namespace UE
 				: PxrSdfPath(MoveTemp(InSdfPath))
 			{
 #if ENABLE_USD_DEBUG_PATH
-				DebugPath = FString(ANSI_TO_TCHAR(PxrSdfPath.Get().GetString().c_str()));
+				DebugPath = FString(UTF8_TO_TCHAR(PxrSdfPath.Get().GetString().c_str()));
 #endif
 			}
 
@@ -67,7 +67,7 @@ namespace UE
 		FScopedUnrealAllocs UnrealAllocs;
 		Impl = MakeUnique<Internal::FSdfPathImpl>(
 			FCString::Strlen(InPath) == 0 ? pxr::SdfPath()	  // USD gives a warning if we construct with an empty string
-										  : pxr::SdfPath(TCHAR_TO_ANSI(InPath))
+										  : pxr::SdfPath(TCHAR_TO_UTF8(InPath))
 		);
 #endif	  // #if USE_USD_SDK
 	}
@@ -250,7 +250,7 @@ namespace UE
 	FString FSdfPath::GetName() const
 	{
 #if USE_USD_SDK
-		return FString(ANSI_TO_TCHAR(Impl->PxrSdfPath.Get().GetName().c_str()));
+		return FString(UTF8_TO_TCHAR(Impl->PxrSdfPath.Get().GetName().c_str()));
 #else
 		return FString();
 #endif	  // #if USE_USD_SDK
@@ -259,7 +259,7 @@ namespace UE
 	FString FSdfPath::GetElementString() const
 	{
 #if USE_USD_SDK
-		return FString(ANSI_TO_TCHAR(Impl->PxrSdfPath.Get().GetElementString().c_str()));
+		return FString(UTF8_TO_TCHAR(Impl->PxrSdfPath.Get().GetElementString().c_str()));
 #else
 		return FString();
 #endif	  // #if USE_USD_SDK
@@ -277,7 +277,7 @@ namespace UE
 	FSdfPath FSdfPath::ReplaceName(const TCHAR* NewLeafName) const
 	{
 #if USE_USD_SDK
-		return FSdfPath(Impl->PxrSdfPath.Get().ReplaceName(pxr::TfToken(TCHAR_TO_ANSI(NewLeafName))));
+		return FSdfPath(Impl->PxrSdfPath.Get().ReplaceName(pxr::TfToken(TCHAR_TO_UTF8(NewLeafName))));
 #else
 		return FSdfPath();
 #endif	  // #if USE_USD_SDK
@@ -313,7 +313,7 @@ namespace UE
 	FSdfPath FSdfPath::AppendChild(const TCHAR* ChildName) const
 	{
 #if USE_USD_SDK
-		return FSdfPath(Impl->PxrSdfPath.Get().AppendChild(pxr::TfToken(TCHAR_TO_ANSI(ChildName))));
+		return FSdfPath(Impl->PxrSdfPath.Get().AppendChild(pxr::TfToken(TCHAR_TO_UTF8(ChildName))));
 #else
 		return FSdfPath();
 #endif	  // #if USE_USD_SDK
@@ -322,7 +322,7 @@ namespace UE
 	FSdfPath FSdfPath::AppendProperty(FName PropertyName) const
 	{
 #if USE_USD_SDK
-		return FSdfPath(Impl->PxrSdfPath.Get().AppendProperty(pxr::TfToken(TCHAR_TO_ANSI(*PropertyName.ToString()))));
+		return FSdfPath(Impl->PxrSdfPath.Get().AppendProperty(pxr::TfToken(TCHAR_TO_UTF8(*PropertyName.ToString()))));
 #else
 		return FSdfPath();
 #endif	  // #if USE_USD_SDK
@@ -340,7 +340,7 @@ namespace UE
 	FString FSdfPath::GetString() const
 	{
 #if USE_USD_SDK
-		return FString(ANSI_TO_TCHAR(Impl->PxrSdfPath.Get().GetString().c_str()));
+		return FString(UTF8_TO_TCHAR(Impl->PxrSdfPath.Get().GetString().c_str()));
 #else
 		return FString();
 #endif	  // #if USE_USD_SDK
@@ -357,7 +357,7 @@ namespace UE
 
 		for (const pxr::SdfPath& UsdPrefix : UsdPrefixes)
 		{
-			Result.Add(UE::FSdfPath{ANSI_TO_TCHAR(UsdPrefix.GetString().c_str())});
+			Result.Add(UE::FSdfPath{UTF8_TO_TCHAR(UsdPrefix.GetString().c_str())});
 		}
 #endif	  // #if USE_USD_SDK
 		return Result;
diff --git a/Engine/Plugins/Runtime/USDCore/Source/UnrealUSDWrapper/Private/UsdWrappers/SdfPrimSpec.cpp b/Engine/Plugins/Runtime/USDCore/Source/UnrealUSDWrapper/Private/UsdWrappers/SdfPrimSpec.cpp
--- a/Engine/Plugins/Runtime/USDCore/Source/UnrealUSDWrapper/Private/UsdWrappers/SdfPrimSpec.cpp
+++ b/Engine/Plugins/Runtime/USDCore/Source/UnrealUSDWrapper/Private/UsdWrappers/SdfPrimSpec.cpp
@@ -212,7 +212,7 @@ namespace UE
 	FName FSdfPrimSpec::GetTypeName() const
 	{
 #if USE_USD_SDK
-		return FName(ANSI_TO_TCHAR(Impl->PxrSdfPrimSpec.Get()->GetTypeName().GetString().c_str()));
+		return FName(UTF8_TO_TCHAR(Impl->PxrSdfPrimSpec.Get()->GetTypeName().GetString().c_str()));
 #else
 		return FName();
 #endif	  // #if USE_USD_SDK
@@ -221,7 +221,7 @@ namespace UE
 	FName FSdfPrimSpec::GetName() const
 	{
 #if USE_USD_SDK
-		return FName(ANSI_TO_TCHAR(Impl->PxrSdfPrimSpec.Get()->GetName().c_str()));
+		return FName(UTF8_TO_TCHAR(Impl->PxrSdfPrimSpec.Get()->GetName().c_str()));
 #else
 		return FName();
 #endif	  // #if USE_USD_SDK
diff --git a/Engine/Plugins/Runtime/USDCore/Source/UnrealUSDWrapper/Private/UsdWrappers/UsdAttribute.cpp b/Engine/Plugins/Runtime/USDCore/Source/UnrealUSDWrapper/Private/UsdWrappers/UsdAttribute.cpp
--- a/Engine/Plugins/Runtime/USDCore/Source/UnrealUSDWrapper/Private/UsdWrappers/UsdAttribute.cpp
+++ b/Engine/Plugins/Runtime/USDCore/Source/UnrealUSDWrapper/Private/UsdWrappers/UsdAttribute.cpp
@@ -155,7 +155,7 @@ namespace UE
 	bool FUsdAttribute::GetMetadata(const TCHAR* Key, UE::FVtValue& Value) const
 	{
 #if USE_USD_SDK
-		return Impl->PxrUsdAttribute.Get().GetMetadata(pxr::TfToken{TCHAR_TO_ANSI(Key)}, &Value.GetUsdValue());
+		return Impl->PxrUsdAttribute.Get().GetMetadata(pxr::TfToken{TCHAR_TO_UTF8(Key)}, &Value.GetUsdValue());
 #else
 		return false;
 #endif	  // #if USE_USD_SDK
@@ -164,7 +164,7 @@ namespace UE
 	bool FUsdAttribute::HasMetadata(const TCHAR* Key) const
 	{
 #if USE_USD_SDK
-		return Impl->PxrUsdAttribute.Get().HasMetadata(pxr::TfToken{TCHAR_TO_ANSI(Key)});
+		return Impl->PxrUsdAttribute.Get().HasMetadata(pxr::TfToken{TCHAR_TO_UTF8(Key)});
 #else
 		return false;
 #endif	  // #if USE_USD_SDK
@@ -173,7 +173,7 @@ namespace UE
 	bool FUsdAttribute::SetMetadata(const TCHAR* Key, const UE::FVtValue& Value) const
 	{
 #if USE_USD_SDK
-		return Impl->PxrUsdAttribute.Get().SetMetadata(pxr::TfToken{TCHAR_TO_ANSI(Key)}, Value.GetUsdValue());
+		return Impl->PxrUsdAttribute.Get().SetMetadata(pxr::TfToken{TCHAR_TO_UTF8(Key)}, Value.GetUsdValue());
 #else
 		return false;
 #endif	  // #if USE_USD_SDK
@@ -182,7 +182,7 @@ namespace UE
 	bool FUsdAttribute::ClearMetadata(const TCHAR* Key) const
 	{
 #if USE_USD_SDK
-		return Impl->PxrUsdAttribute.Get().ClearMetadata(pxr::TfToken{TCHAR_TO_ANSI(Key)});
+		return Impl->PxrUsdAttribute.Get().ClearMetadata(pxr::TfToken{TCHAR_TO_UTF8(Key)});
 #else
 		return false;
 #endif	  // #if USE_USD_SDK
@@ -191,7 +191,7 @@ namespace UE
 	FName FUsdAttribute::GetName() const
 	{
 #if USE_USD_SDK
-		return FName(ANSI_TO_TCHAR(Impl->PxrUsdAttribute.Get().GetName().GetString().c_str()));
+		return FName(UTF8_TO_TCHAR(Impl->PxrUsdAttribute.Get().GetName().GetString().c_str()));
 #else
 		return FName();
 #endif	  // #if USE_USD_SDK
@@ -200,7 +200,7 @@ namespace UE
 	FName FUsdAttribute::GetBaseName() const
 	{
 #if USE_USD_SDK
-		return FName(ANSI_TO_TCHAR(Impl->PxrUsdAttribute.Get().GetBaseName().GetString().c_str()));
+		return FName(UTF8_TO_TCHAR(Impl->PxrUsdAttribute.Get().GetBaseName().GetString().c_str()));
 #else
 		return FName();
 #endif	  // #if USE_USD_SDK
@@ -209,7 +209,7 @@ namespace UE
 	FName FUsdAttribute::GetTypeName() const
 	{
 #if USE_USD_SDK
-		return FName(ANSI_TO_TCHAR(Impl->PxrUsdAttribute.Get().GetTypeName().GetAsToken().GetString().c_str()));
+		return FName(UTF8_TO_TCHAR(Impl->PxrUsdAttribute.Get().GetTypeName().GetAsToken().GetString().c_str()));
 #else
 		return FName();
 #endif	  // #if USE_USD_SDK
@@ -369,7 +369,7 @@ namespace UE
 	}
 
 	template <typename T>
-	bool FUsdAttribute::Get(T& Value, TOptional<double> Time) const 
+	bool FUsdAttribute::Get(T& Value, TOptional<double> Time) const
 	{
 #if USE_USD_SDK
 		FScopedUsdAllocs UsdAllocs;
diff --git a/Engine/Plugins/Runtime/USDCore/Source/UnrealUSDWrapper/Private/UsdWrappers/UsdPayloads.cpp b/Engine/Plugins/Runtime/USDCore/Source/UnrealUSDWrapper/Private/UsdWrappers/UsdPayloads.cpp
--- a/Engine/Plugins/Runtime/USDCore/Source/UnrealUSDWrapper/Private/UsdWrappers/UsdPayloads.cpp
+++ b/Engine/Plugins/Runtime/USDCore/Source/UnrealUSDWrapper/Private/UsdWrappers/UsdPayloads.cpp
@@ -141,7 +141,7 @@ namespace UE
 		{
 			FScopedUsdAllocs UsdAllocs;
 
-			const std::string AssetPath(TCHAR_TO_ANSI(*Payload.AssetPath));
+			const std::string AssetPath(TCHAR_TO_UTF8(*Payload.AssetPath));
 			const pxr::SdfPath UsdPrimPath(Payload.PrimPath);
 			const pxr::SdfLayerOffset UsdLayerOffset(Payload.LayerOffset.Offset, Payload.LayerOffset.Scale);
 
@@ -161,7 +161,7 @@ namespace UE
 		{
 			FScopedUsdAllocs UsdAllocs;
 
-			const std::string UsdIdentifier(TCHAR_TO_ANSI(*Identifier));
+			const std::string UsdIdentifier(TCHAR_TO_UTF8(*Identifier));
 			const pxr::SdfPath UsdPrimPath(PrimPath);
 			const pxr::SdfLayerOffset UsdLayerOffset(LayerOffset.Offset, LayerOffset.Scale);
 
@@ -184,7 +184,7 @@ namespace UE
 		{
 			FScopedUsdAllocs UsdAllocs;
 
-			const std::string UsdIdentifier(TCHAR_TO_ANSI(*Identifier));
+			const std::string UsdIdentifier(TCHAR_TO_UTF8(*Identifier));
 			const pxr::SdfLayerOffset UsdLayerOffset(LayerOffset.Offset, LayerOffset.Scale);
 
 			return Impl->PxrUsdPrim.Get().GetPayloads().AddPayload(UsdIdentifier, UsdLayerOffset, static_cast<pxr::UsdListPosition>(Position));
@@ -218,7 +218,7 @@ namespace UE
 		{
 			FScopedUsdAllocs UsdAllocs;
 
-			const std::string AssetPath(TCHAR_TO_ANSI(*Payload.AssetPath));
+			const std::string AssetPath(TCHAR_TO_UTF8(*Payload.AssetPath));
 			const pxr::SdfPath UsdPrimPath(Payload.PrimPath);
 			const pxr::SdfLayerOffset UsdLayerOffset(Payload.LayerOffset.Offset, Payload.LayerOffset.Scale);
 
@@ -257,7 +257,7 @@ namespace UE
 
 			for (const FSdfPayload& Item : Items)
 			{
-				const std::string AssetPath(TCHAR_TO_ANSI(*Item.AssetPath));
+				const std::string AssetPath(TCHAR_TO_UTF8(*Item.AssetPath));
 				const pxr::SdfPath UsdPrimPath(Item.PrimPath);
 				const pxr::SdfLayerOffset UsdLayerOffset(Item.LayerOffset.Offset, Item.LayerOffset.Scale);
 
diff --git a/Engine/Plugins/Runtime/USDCore/Source/UnrealUSDWrapper/Private/UsdWrappers/UsdPrim.cpp b/Engine/Plugins/Runtime/USDCore/Source/UnrealUSDWrapper/Private/UsdWrappers/UsdPrim.cpp
--- a/Engine/Plugins/Runtime/USDCore/Source/UnrealUSDWrapper/Private/UsdWrappers/UsdPrim.cpp
+++ b/Engine/Plugins/Runtime/USDCore/Source/UnrealUSDWrapper/Private/UsdWrappers/UsdPrim.cpp
@@ -245,7 +245,7 @@ namespace UE
 
 		for (const pxr::TfToken& UsdSchema : UsdAppliedSchemas)
 		{
-			AppliedSchemas.Add(ANSI_TO_TCHAR(UsdSchema.GetString().c_str()));
+			AppliedSchemas.Add(UTF8_TO_TCHAR(UsdSchema.GetString().c_str()));
 		}
 #endif	  // #if USE_USD_SDK
 
@@ -257,7 +257,7 @@ namespace UE
 #if USE_USD_SDK
 		FScopedUsdAllocs Allocs;
 
-		const pxr::TfToken UsdSchemaIdentifier(TCHAR_TO_ANSI(*SchemaIdentifier.ToString()));
+		const pxr::TfToken UsdSchemaIdentifier(TCHAR_TO_UTF8(*SchemaIdentifier.ToString()));
 
 		return Impl->PxrUsdPrim.Get().IsA(UsdSchemaIdentifier);
 #else
@@ -270,7 +270,7 @@ namespace UE
 #if USE_USD_SDK
 		FScopedUsdAllocs Allocs;
 
-		const pxr::TfToken UsdSchemaIdentifier(TCHAR_TO_ANSI(*SchemaIdentifier.ToString()));
+		const pxr::TfToken UsdSchemaIdentifier(TCHAR_TO_UTF8(*SchemaIdentifier.ToString()));
 
 		return Impl->PxrUsdPrim.Get().HasAPI(UsdSchemaIdentifier);
 #else
@@ -283,8 +283,8 @@ namespace UE
 #if USE_USD_SDK
 		FScopedUsdAllocs Allocs;
 
-		const pxr::TfToken UsdSchemaIdentifier(TCHAR_TO_ANSI(*SchemaIdentifier.ToString()));
-		const pxr::TfToken UsdInstanceName(TCHAR_TO_ANSI(*InstanceName.ToString()));
+		const pxr::TfToken UsdSchemaIdentifier(TCHAR_TO_UTF8(*SchemaIdentifier.ToString()));
+		const pxr::TfToken UsdInstanceName(TCHAR_TO_UTF8(*InstanceName.ToString()));
 
 		return Impl->PxrUsdPrim.Get().HasAPI(UsdSchemaIdentifier, UsdInstanceName);
 #else
@@ -310,14 +310,14 @@ namespace UE
 #if USE_USD_SDK
 		FScopedUsdAllocs Allocs;
 
-		const pxr::TfToken UsdSchemaIdentifier(TCHAR_TO_ANSI(*SchemaIdentifier.ToString()));
+		const pxr::TfToken UsdSchemaIdentifier(TCHAR_TO_UTF8(*SchemaIdentifier.ToString()));
 
 		std::string WhyNot;
 		bResult = Impl->PxrUsdPrim.Get().CanApplyAPI(UsdSchemaIdentifier, OutWhyNot != nullptr ? &WhyNot : nullptr);
 
 		if (!bResult && OutWhyNot != nullptr)
 		{
-			*OutWhyNot = ANSI_TO_TCHAR(WhyNot.c_str());
+			*OutWhyNot = UTF8_TO_TCHAR(WhyNot.c_str());
 		}
 #endif	  // #if USE_USD_SDK
 
@@ -331,15 +331,15 @@ namespace UE
 #if USE_USD_SDK
 		FScopedUsdAllocs Allocs;
 
-		const pxr::TfToken UsdSchemaIdentifier(TCHAR_TO_ANSI(*SchemaIdentifier.ToString()));
-		const pxr::TfToken UsdInstanceName(TCHAR_TO_ANSI(*InstanceName.ToString()));
+		const pxr::TfToken UsdSchemaIdentifier(TCHAR_TO_UTF8(*SchemaIdentifier.ToString()));
+		const pxr::TfToken UsdInstanceName(TCHAR_TO_UTF8(*InstanceName.ToString()));
 
 		std::string WhyNot;
 		bResult = Impl->PxrUsdPrim.Get().CanApplyAPI(UsdSchemaIdentifier, UsdInstanceName, OutWhyNot != nullptr ? &WhyNot : nullptr);
 
 		if (!bResult && OutWhyNot != nullptr)
 		{
-			*OutWhyNot = ANSI_TO_TCHAR(WhyNot.c_str());
+			*OutWhyNot = UTF8_TO_TCHAR(WhyNot.c_str());
 		}
 #endif	  // #if USE_USD_SDK
 
@@ -351,7 +351,7 @@ namespace UE
 #if USE_USD_SDK
 		FScopedUsdAllocs Allocs;
 
-		const pxr::TfToken UsdSchemaIdentifier(TCHAR_TO_ANSI(*SchemaIdentifier.ToString()));
+		const pxr::TfToken UsdSchemaIdentifier(TCHAR_TO_UTF8(*SchemaIdentifier.ToString()));
 
 		return Impl->PxrUsdPrim.Get().ApplyAPI(UsdSchemaIdentifier);
 #else
@@ -364,8 +364,8 @@ namespace UE
 #if USE_USD_SDK
 		FScopedUsdAllocs Allocs;
 
-		const pxr::TfToken UsdSchemaIdentifier(TCHAR_TO_ANSI(*SchemaIdentifier.ToString()));
-		const pxr::TfToken UsdInstanceName(TCHAR_TO_ANSI(*InstanceName.ToString()));
+		const pxr::TfToken UsdSchemaIdentifier(TCHAR_TO_UTF8(*SchemaIdentifier.ToString()));
+		const pxr::TfToken UsdInstanceName(TCHAR_TO_UTF8(*InstanceName.ToString()));
 
 		return Impl->PxrUsdPrim.Get().ApplyAPI(UsdSchemaIdentifier, UsdInstanceName);
 #else
@@ -378,7 +378,7 @@ namespace UE
 #if USE_USD_SDK
 		FScopedUsdAllocs Allocs;
 
-		const pxr::TfToken UsdSchemaIdentifier(TCHAR_TO_ANSI(*SchemaIdentifier.ToString()));
+		const pxr::TfToken UsdSchemaIdentifier(TCHAR_TO_UTF8(*SchemaIdentifier.ToString()));
 
 		return Impl->PxrUsdPrim.Get().RemoveAPI(UsdSchemaIdentifier);
 #else
@@ -391,8 +391,8 @@ namespace UE
 #if USE_USD_SDK
 		FScopedUsdAllocs Allocs;
 
-		const pxr::TfToken UsdSchemaIdentifier(TCHAR_TO_ANSI(*SchemaIdentifier.ToString()));
-		const pxr::TfToken UsdInstanceName(TCHAR_TO_ANSI(*InstanceName.ToString()));
+		const pxr::TfToken UsdSchemaIdentifier(TCHAR_TO_UTF8(*SchemaIdentifier.ToString()));
+		const pxr::TfToken UsdInstanceName(TCHAR_TO_UTF8(*InstanceName.ToString()));
 
 		return Impl->PxrUsdPrim.Get().RemoveAPI(UsdSchemaIdentifier, UsdInstanceName);
 #else
@@ -421,7 +421,7 @@ namespace UE
 	FName FUsdPrim::GetName() const
 	{
 #if USE_USD_SDK
-		return FName(ANSI_TO_TCHAR(Impl->PxrUsdPrim.Get().GetName().GetString().c_str()));
+		return FName(UTF8_TO_TCHAR(Impl->PxrUsdPrim.Get().GetName().GetString().c_str()));
 #else
 		return FName();
 #endif	  // #if USE_USD_SDK
@@ -430,7 +430,7 @@ namespace UE
 	FName FUsdPrim::GetTypeName() const
 	{
 #if USE_USD_SDK
-		return FName(ANSI_TO_TCHAR(Impl->PxrUsdPrim.Get().GetTypeName().GetString().c_str()));
+		return FName(UTF8_TO_TCHAR(Impl->PxrUsdPrim.Get().GetTypeName().GetString().c_str()));
 #else
 		return FName();
 #endif	  // #if USE_USD_SDK
@@ -440,7 +440,7 @@ namespace UE
 	{
 #if USE_USD_SDK
 		FScopedUsdAllocs UsdAllocs;
-		return Impl->PxrUsdPrim.Get().SetTypeName(pxr::TfToken(TCHAR_TO_ANSI(*TypeName.ToString())));
+		return Impl->PxrUsdPrim.Get().SetTypeName(pxr::TfToken(TCHAR_TO_UTF8(*TypeName.ToString())));
 #else
 		return false;
 #endif	  // #if USE_USD_SDK
@@ -530,7 +530,7 @@ namespace UE
 #if USE_USD_SDK
 		FScopedUsdAllocs UsdAllocs;
 
-		const std::string UsdVariantSetName(TCHAR_TO_ANSI(*VariantSetName));
+		const std::string UsdVariantSetName(TCHAR_TO_UTF8(*VariantSetName));
 
 		return FUsdVariantSet(Impl->PxrUsdPrim.Get(), UsdVariantSetName);
 #else
@@ -639,7 +639,7 @@ namespace UE
 #if USE_USD_SDK
 		FScopedUsdAllocs UsdAllocs;
 
-		return FUsdAttribute(Impl->PxrUsdPrim.Get().GetAttribute(pxr::TfToken(TCHAR_TO_ANSI(AttrName))));
+		return FUsdAttribute(Impl->PxrUsdPrim.Get().GetAttribute(pxr::TfToken(TCHAR_TO_UTF8(AttrName))));
 #else
 		return FUsdAttribute{};
 #endif	  // #if USE_USD_SDK
@@ -648,7 +648,7 @@ namespace UE
 	bool FUsdPrim::HasAttribute(const TCHAR* AttrName) const
 	{
 #if USE_USD_SDK
-		return Impl->PxrUsdPrim.Get().HasAttribute(pxr::TfToken(TCHAR_TO_ANSI(AttrName)));
+		return Impl->PxrUsdPrim.Get().HasAttribute(pxr::TfToken(TCHAR_TO_UTF8(AttrName)));
 #else
 		return false;
 #endif	  // #if USE_USD_SDK
@@ -660,8 +660,8 @@ namespace UE
 		FScopedUsdAllocs UsdAllocs;
 
 		return FUsdAttribute(Impl->PxrUsdPrim.Get().CreateAttribute(
-			pxr::TfToken(TCHAR_TO_ANSI(AttrName)),
-			pxr::SdfSchema::GetInstance().FindType(TCHAR_TO_ANSI(*TypeName.ToString()))
+			pxr::TfToken(TCHAR_TO_UTF8(AttrName)),
+			pxr::SdfSchema::GetInstance().FindType(TCHAR_TO_UTF8(*TypeName.ToString()))
 		));
 #else
 		return FUsdAttribute{};
@@ -673,7 +673,7 @@ namespace UE
 #if USE_USD_SDK
 		FScopedUsdAllocs UsdAllocs;
 
-		return FUsdRelationship(Impl->PxrUsdPrim.Get().CreateRelationship(pxr::TfToken(TCHAR_TO_ANSI(RelName)), bCustom));
+		return FUsdRelationship(Impl->PxrUsdPrim.Get().CreateRelationship(pxr::TfToken(TCHAR_TO_UTF8(RelName)), bCustom));
 #else
 		return FUsdRelationship{};
 #endif	  // #if USE_USD_SDK
@@ -682,7 +682,7 @@ namespace UE
 	FUsdRelationship FUsdPrim::GetRelationship(const TCHAR* RelationshipName) const
 	{
 #if USE_USD_SDK
-		return FUsdRelationship(Impl->PxrUsdPrim.Get().GetRelationship(pxr::TfToken(TCHAR_TO_ANSI(RelationshipName))));
+		return FUsdRelationship(Impl->PxrUsdPrim.Get().GetRelationship(pxr::TfToken(TCHAR_TO_UTF8(RelationshipName))));
 #else
 		return FUsdRelationship();
 #endif	  // #if USE_USD_SDK
@@ -707,7 +707,7 @@ namespace UE
 	bool FUsdPrim::HasRelationship(const TCHAR* RelName) const
 	{
 #if USE_USD_SDK
-		return Impl->PxrUsdPrim.Get().HasRelationship(pxr::TfToken(TCHAR_TO_ANSI(RelName)));
+		return Impl->PxrUsdPrim.Get().HasRelationship(pxr::TfToken(TCHAR_TO_UTF8(RelName)));
 #else
 		return false;
 #endif	  // #if USE_USD_SDK
@@ -717,7 +717,7 @@ namespace UE
 	{
 #if USE_USD_SDK
 		FScopedUsdAllocs UsdAllocs;
-		return Impl->PxrUsdPrim.Get().RemoveProperty(pxr::TfToken(TCHAR_TO_ANSI(*PropName.ToString())));
+		return Impl->PxrUsdPrim.Get().RemoveProperty(pxr::TfToken(TCHAR_TO_UTF8(*PropName.ToString())));
 #else
 		return false;
 #endif	  // #if USE_USD_SDK
diff --git a/Engine/Plugins/Runtime/USDCore/Source/UnrealUSDWrapper/Private/UsdWrappers/UsdReferences.cpp b/Engine/Plugins/Runtime/USDCore/Source/UnrealUSDWrapper/Private/UsdWrappers/UsdReferences.cpp
--- a/Engine/Plugins/Runtime/USDCore/Source/UnrealUSDWrapper/Private/UsdWrappers/UsdReferences.cpp
+++ b/Engine/Plugins/Runtime/USDCore/Source/UnrealUSDWrapper/Private/UsdWrappers/UsdReferences.cpp
@@ -158,7 +158,7 @@ namespace UE
 		{
 			FScopedUsdAllocs UsdAllocs;
 
-			const std::string AssetPath(TCHAR_TO_ANSI(*Reference.AssetPath));
+			const std::string AssetPath(TCHAR_TO_UTF8(*Reference.AssetPath));
 			const pxr::SdfPath UsdPrimPath(Reference.PrimPath);
 			const pxr::SdfLayerOffset UsdLayerOffset(Reference.LayerOffset.Offset, Reference.LayerOffset.Scale);
 
@@ -183,7 +183,7 @@ namespace UE
 		{
 			FScopedUsdAllocs UsdAllocs;
 
-			const std::string UsdIdentifier(TCHAR_TO_ANSI(*Identifier));
+			const std::string UsdIdentifier(TCHAR_TO_UTF8(*Identifier));
 			const pxr::SdfPath UsdPrimPath(PrimPath);
 			const pxr::SdfLayerOffset UsdLayerOffset(LayerOffset.Offset, LayerOffset.Scale);
 
@@ -206,7 +206,7 @@ namespace UE
 		{
 			FScopedUsdAllocs UsdAllocs;
 
-			const std::string UsdIdentifier(TCHAR_TO_ANSI(*Identifier));
+			const std::string UsdIdentifier(TCHAR_TO_UTF8(*Identifier));
 			const pxr::SdfLayerOffset UsdLayerOffset(LayerOffset.Offset, LayerOffset.Scale);
 
 			return Impl->PxrUsdPrim.Get().GetReferences().AddReference(UsdIdentifier, UsdLayerOffset, static_cast<pxr::UsdListPosition>(Position));
@@ -244,7 +244,7 @@ namespace UE
 		{
 			FScopedUsdAllocs UsdAllocs;
 
-			const std::string AssetPath(TCHAR_TO_ANSI(*Reference.AssetPath));
+			const std::string AssetPath(TCHAR_TO_UTF8(*Reference.AssetPath));
 			const pxr::SdfPath UsdPrimPath(Reference.PrimPath);
 			const pxr::SdfLayerOffset UsdLayerOffset(Reference.LayerOffset.Offset, Reference.LayerOffset.Scale);
 
@@ -283,7 +283,7 @@ namespace UE
 
 			for (const FSdfReference& Item : Items)
 			{
-				const std::string AssetPath(TCHAR_TO_ANSI(*Item.AssetPath));
+				const std::string AssetPath(TCHAR_TO_UTF8(*Item.AssetPath));
 				const pxr::SdfPath UsdPrimPath(Item.PrimPath);
 				const pxr::SdfLayerOffset UsdLayerOffset(Item.LayerOffset.Offset, Item.LayerOffset.Scale);
 
diff --git a/Engine/Plugins/Runtime/USDCore/Source/UnrealUSDWrapper/Private/UsdWrappers/UsdRelationship.cpp b/Engine/Plugins/Runtime/USDCore/Source/UnrealUSDWrapper/Private/UsdWrappers/UsdRelationship.cpp
--- a/Engine/Plugins/Runtime/USDCore/Source/UnrealUSDWrapper/Private/UsdWrappers/UsdRelationship.cpp
+++ b/Engine/Plugins/Runtime/USDCore/Source/UnrealUSDWrapper/Private/UsdWrappers/UsdRelationship.cpp
@@ -152,7 +152,7 @@ namespace UE
 	bool FUsdRelationship::GetMetadata(const TCHAR* Key, UE::FVtValue& Value) const
 	{
 #if USE_USD_SDK
-		return Impl->PxrUsdRelationship.Get().GetMetadata(pxr::TfToken{TCHAR_TO_ANSI(Key)}, &Value.GetUsdValue());
+		return Impl->PxrUsdRelationship.Get().GetMetadata(pxr::TfToken{TCHAR_TO_UTF8(Key)}, &Value.GetUsdValue());
 #else
 		return false;
 #endif	  // #if USE_USD_SDK
@@ -161,7 +161,7 @@ namespace UE
 	bool FUsdRelationship::HasMetadata(const TCHAR* Key) const
 	{
 #if USE_USD_SDK
-		return Impl->PxrUsdRelationship.Get().HasMetadata(pxr::TfToken{TCHAR_TO_ANSI(Key)});
+		return Impl->PxrUsdRelationship.Get().HasMetadata(pxr::TfToken{TCHAR_TO_UTF8(Key)});
 #else
 		return false;
 #endif	  // #if USE_USD_SDK
@@ -170,7 +170,7 @@ namespace UE
 	bool FUsdRelationship::SetMetadata(const TCHAR* Key, const UE::FVtValue& Value) const
 	{
 #if USE_USD_SDK
-		return Impl->PxrUsdRelationship.Get().SetMetadata(pxr::TfToken{TCHAR_TO_ANSI(Key)}, Value.GetUsdValue());
+		return Impl->PxrUsdRelationship.Get().SetMetadata(pxr::TfToken{TCHAR_TO_UTF8(Key)}, Value.GetUsdValue());
 #else
 		return false;
 #endif	  // #if USE_USD_SDK
@@ -179,7 +179,7 @@ namespace UE
 	bool FUsdRelationship::ClearMetadata(const TCHAR* Key) const
 	{
 #if USE_USD_SDK
-		return Impl->PxrUsdRelationship.Get().ClearMetadata(pxr::TfToken{TCHAR_TO_ANSI(Key)});
+		return Impl->PxrUsdRelationship.Get().ClearMetadata(pxr::TfToken{TCHAR_TO_UTF8(Key)});
 #else
 		return false;
 #endif	  // #if USE_USD_SDK
diff --git a/Engine/Plugins/Runtime/USDCore/Source/UnrealUSDWrapper/Private/UsdWrappers/UsdStage.cpp b/Engine/Plugins/Runtime/USDCore/Source/UnrealUSDWrapper/Private/UsdWrappers/UsdStage.cpp
--- a/Engine/Plugins/Runtime/USDCore/Source/UnrealUSDWrapper/Private/UsdWrappers/UsdStage.cpp
+++ b/Engine/Plugins/Runtime/USDCore/Source/UnrealUSDWrapper/Private/UsdWrappers/UsdStage.cpp
@@ -263,10 +263,10 @@ namespace UE
 			std::map<std::string, std::string> UsdFileFormatArguments;
 			for (const TPair<FString, FString>& Pair : FileFormatArguments)
 			{
-				UsdFileFormatArguments.insert({TCHAR_TO_ANSI(*Pair.Key), TCHAR_TO_ANSI(*Pair.Value)});
+				UsdFileFormatArguments.insert({TCHAR_TO_UTF8(*Pair.Key), TCHAR_TO_UTF8(*Pair.Value)});
 			}
 
-			return Ptr->Export(TCHAR_TO_ANSI(FileName), bAddSourceFileComment, UsdFileFormatArguments);
+			return Ptr->Export(TCHAR_TO_UTF8(FileName), bAddSourceFileComment, UsdFileFormatArguments);
 		}
 #endif	  // #if USE_USD_SDK
 
@@ -423,14 +423,14 @@ namespace UE
 			UsdMuteLayers.reserve(MuteLayers.Num());
 			for (const FString& MuteLayer : MuteLayers)
 			{
-				UsdMuteLayers.push_back(TCHAR_TO_ANSI(*MuteLayer));
+				UsdMuteLayers.push_back(TCHAR_TO_UTF8(*MuteLayer));
 			}
 
 			std::vector<std::string> UsdUnmuteLayers;
 			UsdUnmuteLayers.reserve(UnmuteLayers.Num());
 			for (const FString& UnmuteLayer : UnmuteLayers)
 			{
-				UsdUnmuteLayers.push_back(TCHAR_TO_ANSI(*UnmuteLayer));
+				UsdUnmuteLayers.push_back(TCHAR_TO_UTF8(*UnmuteLayer));
 			}
 
 			Ptr->MuteAndUnmuteLayers(UsdMuteLayers, UsdUnmuteLayers);
@@ -446,7 +446,7 @@ namespace UE
 		{
 			FScopedUsdAllocs UsdAllocs;
 
-			return Ptr->IsLayerMuted(TCHAR_TO_ANSI(*LayerIdentifier));
+			return Ptr->IsLayerMuted(TCHAR_TO_UTF8(*LayerIdentifier));
 		}
 #endif	  // #if USE_USD_SDK
 
@@ -475,7 +475,7 @@ namespace UE
 #if USE_USD_SDK
 		if (PtrType& Ptr = Impl->GetInner())
 		{
-			return Ptr->GetMetadata(pxr::TfToken{TCHAR_TO_ANSI(Key)}, &Value.GetUsdValue());
+			return Ptr->GetMetadata(pxr::TfToken{TCHAR_TO_UTF8(Key)}, &Value.GetUsdValue());
 		}
 #endif	  // #if USE_USD_SDK
 
@@ -488,7 +488,7 @@ namespace UE
 #if USE_USD_SDK
 		if (PtrType& Ptr = Impl->GetInner())
 		{
-			return Ptr->HasMetadata(pxr::TfToken{TCHAR_TO_ANSI(Key)});
+			return Ptr->HasMetadata(pxr::TfToken{TCHAR_TO_UTF8(Key)});
 		}
 #endif	  // #if USE_USD_SDK
 
@@ -501,7 +501,7 @@ namespace UE
 #if USE_USD_SDK
 		if (PtrType& Ptr = Impl->GetInner())
 		{
-			return Ptr->SetMetadata(pxr::TfToken{TCHAR_TO_ANSI(Key)}, Value.GetUsdValue());
+			return Ptr->SetMetadata(pxr::TfToken{TCHAR_TO_UTF8(Key)}, Value.GetUsdValue());
 		}
 #endif	  // #if USE_USD_SDK
 
@@ -514,7 +514,7 @@ namespace UE
 #if USE_USD_SDK
 		if (PtrType& Ptr = Impl->GetInner())
 		{
-			return Ptr->ClearMetadata(pxr::TfToken{TCHAR_TO_ANSI(Key)});
+			return Ptr->ClearMetadata(pxr::TfToken{TCHAR_TO_UTF8(Key)});
 		}
 #endif	  // #if USE_USD_SDK
 
@@ -690,7 +690,7 @@ namespace UE
 #if USE_USD_SDK
 		if (PtrType& Ptr = Impl->GetInner())
 		{
-			return FUsdPrim(Ptr->DefinePrim(Path, pxr::TfToken(TCHAR_TO_ANSI(TypeName))));
+			return FUsdPrim(Ptr->DefinePrim(Path, pxr::TfToken(TCHAR_TO_UTF8(TypeName))));
 		}
 #endif	  // #if USE_USD_SDK
 
diff --git a/Engine/Plugins/Runtime/USDCore/Source/UnrealUSDWrapper/Private/UsdWrappers/UsdVariantSets.cpp b/Engine/Plugins/Runtime/USDCore/Source/UnrealUSDWrapper/Private/UsdWrappers/UsdVariantSets.cpp
--- a/Engine/Plugins/Runtime/USDCore/Source/UnrealUSDWrapper/Private/UsdWrappers/UsdVariantSets.cpp
+++ b/Engine/Plugins/Runtime/USDCore/Source/UnrealUSDWrapper/Private/UsdWrappers/UsdVariantSets.cpp
@@ -174,7 +174,7 @@ namespace UE
 		pxr::UsdVariantSet UsdVariantSet = Impl->GetInner();
 		if (UsdVariantSet.IsValid())
 		{
-			const std::string UsdVariantName(TCHAR_TO_ANSI(*VariantName));
+			const std::string UsdVariantName(TCHAR_TO_UTF8(*VariantName));
 
 			return UsdVariantSet.AddVariant(UsdVariantName, static_cast<pxr::UsdListPosition>(Position));
 		}
@@ -203,7 +203,7 @@ namespace UE
 
 			for (const std::string& UsdVariantName : UsdVariantNames)
 			{
-				VariantNames.Emplace(ANSI_TO_TCHAR(UsdVariantName.c_str()));
+				VariantNames.Emplace(UTF8_TO_TCHAR(UsdVariantName.c_str()));
 			}
 		}
 #endif	  // #if USE_USD_SDK
@@ -219,7 +219,7 @@ namespace UE
 		pxr::UsdVariantSet UsdVariantSet = Impl->GetInner();
 		if (UsdVariantSet.IsValid())
 		{
-			const std::string UsdVariantName(TCHAR_TO_ANSI(*VariantName));
+			const std::string UsdVariantName(TCHAR_TO_UTF8(*VariantName));
 
 			return UsdVariantSet.HasAuthoredVariant(UsdVariantName);
 		}
@@ -238,7 +238,7 @@ namespace UE
 		{
 			const std::string UsdVariantSelection = UsdVariantSet.GetVariantSelection();
 
-			return FString(ANSI_TO_TCHAR(UsdVariantSelection.c_str()));
+			return FString(UTF8_TO_TCHAR(UsdVariantSelection.c_str()));
 		}
 #endif	  // #if USE_USD_SDK
 
@@ -258,7 +258,7 @@ namespace UE
 
 			if (Value != nullptr)
 			{
-				*Value = ANSI_TO_TCHAR(UsdValue.c_str());
+				*Value = UTF8_TO_TCHAR(UsdValue.c_str());
 			}
 
 			return bResult;
@@ -276,7 +276,7 @@ namespace UE
 		pxr::UsdVariantSet UsdVariantSet = Impl->GetInner();
 		if (UsdVariantSet.IsValid())
 		{
-			const std::string UsdVariantName(TCHAR_TO_ANSI(*VariantName));
+			const std::string UsdVariantName(TCHAR_TO_UTF8(*VariantName));
 
 			return UsdVariantSet.SetVariantSelection(UsdVariantName);
 		}
@@ -340,7 +340,7 @@ namespace UE
 		{
 			const std::string UsdVariantName = UsdVariantSet.GetName();
 
-			return FString(ANSI_TO_TCHAR(UsdVariantName.c_str()));
+			return FString(UTF8_TO_TCHAR(UsdVariantName.c_str()));
 		}
 #endif	  // #if USE_USD_SDK
 
@@ -427,7 +427,7 @@ namespace UE
 
 		pxr::UsdVariantSets UsdVariantSets = Impl->GetInner();
 
-		const std::string UsdVariantName(TCHAR_TO_ANSI(*VariantName));
+		const std::string UsdVariantName(TCHAR_TO_UTF8(*VariantName));
 
 		return FUsdVariantSet(UsdVariantSets.AddVariantSet(UsdVariantName, static_cast<pxr::UsdListPosition>(Position)));
 #else
@@ -454,7 +454,7 @@ namespace UE
 
 		for (const std::string& UsdVariantSetName : UsdVariantSetNames)
 		{
-			VariantSetNames.Emplace(ANSI_TO_TCHAR(UsdVariantSetName.c_str()));
+			VariantSetNames.Emplace(UTF8_TO_TCHAR(UsdVariantSetName.c_str()));
 		}
 #endif	  // #if USE_USD_SDK
 
@@ -468,7 +468,7 @@ namespace UE
 
 		pxr::UsdVariantSets UsdVariantSets = Impl->GetInner();
 
-		const std::string UsdVariantSetName(TCHAR_TO_ANSI(*VariantSetName));
+		const std::string UsdVariantSetName(TCHAR_TO_UTF8(*VariantSetName));
 
 		return FUsdVariantSet(UsdVariantSets.GetVariantSet(UsdVariantSetName));
 #else
@@ -483,7 +483,7 @@ namespace UE
 
 		pxr::UsdVariantSets UsdVariantSets = Impl->GetInner();
 
-		const std::string UsdVariantSetName(TCHAR_TO_ANSI(*VariantSetName));
+		const std::string UsdVariantSetName(TCHAR_TO_UTF8(*VariantSetName));
 
 		return UsdVariantSets.HasVariantSet(UsdVariantSetName);
 #else
@@ -498,11 +498,11 @@ namespace UE
 
 		pxr::UsdVariantSets UsdVariantSets = Impl->GetInner();
 
-		const std::string UsdVariantSetName(TCHAR_TO_ANSI(*VariantSetName));
+		const std::string UsdVariantSetName(TCHAR_TO_UTF8(*VariantSetName));
 
 		const std::string UsdVariantSelection = UsdVariantSets.GetVariantSelection(UsdVariantSetName);
 
-		return FString(ANSI_TO_TCHAR(UsdVariantSelection.c_str()));
+		return FString(UTF8_TO_TCHAR(UsdVariantSelection.c_str()));
 #else
 		return FString{};
 #endif	  // #if USE_USD_SDK
@@ -515,8 +515,8 @@ namespace UE
 
 		pxr::UsdVariantSets UsdVariantSets = Impl->GetInner();
 
-		const std::string UsdVariantSetName(TCHAR_TO_ANSI(*VariantSetName));
-		const std::string UsdVariantName(TCHAR_TO_ANSI(*VariantName));
+		const std::string UsdVariantSetName(TCHAR_TO_UTF8(*VariantSetName));
+		const std::string UsdVariantName(TCHAR_TO_UTF8(*VariantName));
 
 		return UsdVariantSets.SetSelection(UsdVariantSetName, UsdVariantName);
 #else
@@ -543,7 +543,7 @@ namespace UE
 
 		for (const std::pair<const std::string, std::string>& UsdVariantSelection : UsdVariantSelections)
 		{
-			VariantSelections.Emplace(ANSI_TO_TCHAR(UsdVariantSelection.first.c_str()), ANSI_TO_TCHAR(UsdVariantSelection.second.c_str()));
+			VariantSelections.Emplace(UTF8_TO_TCHAR(UsdVariantSelection.first.c_str()), UTF8_TO_TCHAR(UsdVariantSelection.second.c_str()));
 		}
 #endif	  // #if USE_USD_SDK
 
diff --git a/Engine/Plugins/Runtime/USDCore/Source/UnrealUSDWrapper/Private/UsdWrappers/VtValue.cpp b/Engine/Plugins/Runtime/USDCore/Source/UnrealUSDWrapper/Private/UsdWrappers/VtValue.cpp
--- a/Engine/Plugins/Runtime/USDCore/Source/UnrealUSDWrapper/Private/UsdWrappers/VtValue.cpp
+++ b/Engine/Plugins/Runtime/USDCore/Source/UnrealUSDWrapper/Private/UsdWrappers/VtValue.cpp
@@ -127,7 +127,7 @@ namespace UE
 #if USE_USD_SDK
 		FScopedUsdAllocs UsdAllocs;
 
-		return FString(ANSI_TO_TCHAR(Impl->PxrVtValue.Get().GetTypeName().c_str()));
+		return FString(UTF8_TO_TCHAR(Impl->PxrVtValue.Get().GetTypeName().c_str()));
 #else
 		return FString();
 #endif	  // #if USE_USD_SDK
