USD: Patch plugInfo.json files only when needed, but check all of them every time.

diff --git a/Engine/Plugins/Runtime/USDCore/Source/USDClasses/Private/USDClassesModule.cpp b/Engine/Plugins/Runtime/USDCore/Source/USDClasses/Private/USDClassesModule.cpp
--- a/Engine/Plugins/Runtime/USDCore/Source/USDClasses/Private/USDClassesModule.cpp
+++ b/Engine/Plugins/Runtime/USDCore/Source/USDClasses/Private/USDClassesModule.cpp
@@ -199,8 +199,11 @@ void IUsdClassesModule::UpdatePlugInfoFiles(const FString& PluginDirectory, cons
 
 			// e.g. "../../../../../../Win64/sdf.dll"
 			FString NewLibraryPath = FPaths::Combine(PathToDll, DllName);
-			(*PluginEntryObjectPtr)->SetField(TEXT("LibraryPath"), MakeShared<FJsonValueString>(NewLibraryPath));
-			bUpdatedFile = true;
+			if (!FPaths::IsSamePath(NewLibraryPath, LibraryPath))
+			{
+				bUpdatedFile = true;
+				(*PluginEntryObjectPtr)->SetField(TEXT("LibraryPath"), MakeShared<FJsonValueString>(NewLibraryPath));
+			}
 		}
 
 		// Serialize them back to JSON if we changed anything
