USD: Add support for more UsdGeomCamera parameters.

#virtualized

diff --git a/Engine/Plugins/Importers/USDImporter/Source/USDStage/Private/USDLevelSequenceHelper.cpp b/Engine/Plugins/Importers/USDImporter/Source/USDStage/Private/USDLevelSequenceHelper.cpp
--- a/Engine/Plugins/Importers/USDImporter/Source/USDStage/Private/USDLevelSequenceHelper.cpp
+++ b/Engine/Plugins/Importers/USDImporter/Source/USDStage/Private/USDLevelSequenceHelper.cpp
@@ -110,7 +110,7 @@
 
 #if USE_USD_SDK
 
-namespace UsdLevelSequenceHelperImpl
+namespace UE::USDLevelSequenceHelper::Private
 {
 	/**
 	 * Similar to FrameRate.AsFrameNumber(TimeSeconds) except that it uses RoundToDouble instead of FloorToDouble, to
@@ -276,9 +276,9 @@ namespace UsdLevelSequenceHelperImpl
 	{
 		const bool bFocusIfOpen = false;
 		IAssetEditorInstance* AssetEditor = GEditor ? GEditor->GetEditorSubsystem<UAssetEditorSubsystem>()->FindEditorForAsset(
-												LevelSequence,
-												bFocusIfOpen
-											)
+														  LevelSequence,
+														  bFocusIfOpen
+													  )
 													: nullptr;
 		ILevelSequenceEditorToolkit* LevelSequenceEditor = static_cast<ILevelSequenceEditorToolkit*>(AssetEditor);
 		return LevelSequenceEditor ? LevelSequenceEditor->GetSequencer() : nullptr;
@@ -760,7 +760,7 @@ cleanup:
 		}
 	}
 #endif	  // WITH_EDITOR
-}	 // namespace UsdLevelSequenceHelperImpl
+}	 // namespace UE::USDLevelSequenceHelper::Private
 
 class FUsdLevelSequenceHelperImpl : private FGCObject
 {
@@ -1317,14 +1317,13 @@ void FUsdLevelSequenceHelperImpl::OnStageActorRenamed()
 		return;
 	}
 
-	FMovieScenePossessable NewPossessable
-	{
+	FMovieScenePossessable NewPossessable{
 #if WITH_EDITOR
 		StageActorPtr->GetActorLabel(),
 #else
 		StageActorPtr->GetName(),
 #endif	  // WITH_EDITOR
-			StageActorPtr->GetClass()
+		StageActorPtr->GetClass()
 	};
 	FGuid NewId = NewPossessable.GetGuid();
 
@@ -1445,10 +1444,10 @@ ULevelSequence* FUsdLevelSequenceHelperImpl::FindOrAddSequenceForLayer(
 		const bool bIsImporting = StageActor.IsExplicitlyNull() && PrimLinkCache.IsValid();
 		FName UniqueSequenceName = bIsImporting ? *UsdUnreal::ObjectUtils::SanitizeObjectName(FPaths::GetBaseFilename(SequenceDisplayName))
 												: MakeUniqueObjectName(
-													GetTransientPackage(),
-													ULevelSequence::StaticClass(),
-													*UsdUnreal::ObjectUtils::SanitizeObjectName(FPaths::GetBaseFilename(SequenceDisplayName))
-												);
+													  GetTransientPackage(),
+													  ULevelSequence::StaticClass(),
+													  *UsdUnreal::ObjectUtils::SanitizeObjectName(FPaths::GetBaseFilename(SequenceDisplayName))
+												  );
 
 		Sequence = NewObject<ULevelSequence>(GetTransientPackage(), UniqueSequenceName, FUsdLevelSequenceHelperImpl::DefaultObjFlags);
 		Sequence->Initialize();
@@ -1465,7 +1464,7 @@ ULevelSequence* FUsdLevelSequenceHelperImpl::FindOrAddSequenceForLayer(
 
 		// Here we abuse the FScopedReadOnlyDisable so that we can use the code in its destructor to set
 		// Sequence to ReadOnly if Layer doesn't belong to UsdStage's local layer stack
-		UsdLevelSequenceHelperImpl::FScopedReadOnlyDisable SetToReadOnly{MovieScene, Layer, UsdStage};
+		UE::USDLevelSequenceHelper::Private::FScopedReadOnlyDisable SetToReadOnly{MovieScene, Layer, UsdStage};
 
 		const FLayerTimeInfo LayerTimeInfo = FindOrAddLayerTimeInfo(Layer);
 
@@ -1615,14 +1614,14 @@ void FUsdLevelSequenceHelperImpl::CreateSubSequenceSection(ULevelSequence& Seque
 		const double SubStartTimeSeconds = SubLayerOffset.Offset / TimeCodesPerSecond;
 		const double SubEndTimeSeconds = SubStartTimeSeconds + SubDurationSeconds;
 
-		StartFrame = UsdLevelSequenceHelperImpl::RoundAsFrameNumber(TickResolution, SubStartTimeSeconds);
-		const FFrameNumber EndFrame = UsdLevelSequenceHelperImpl::RoundAsFrameNumber(TickResolution, SubEndTimeSeconds);
+		StartFrame = UE::USDLevelSequenceHelper::Private::RoundAsFrameNumber(TickResolution, SubStartTimeSeconds);
+		const FFrameNumber EndFrame = UE::USDLevelSequenceHelper::Private::RoundAsFrameNumber(TickResolution, SubEndTimeSeconds);
 
 		// Don't clip subsections with their duration, so that the root layer's [startTimeCode, endTimeCode] range is the only thing clipping
 		// anything, as this is how USD seems to behave. Even if a middle sublayer has startTimeCode == endTimeCode, its animations
 		// (or its child sublayers') won't be clipped by it and play according to the stage's range
 		const double StageEndTimeSeconds = UsdStage.GetEndTimeCode() / UsdStage.GetTimeCodesPerSecond();
-		const FFrameNumber StageEndFrame = UsdLevelSequenceHelperImpl::RoundAsFrameNumber(TickResolution, StageEndTimeSeconds);
+		const FFrameNumber StageEndFrame = UE::USDLevelSequenceHelper::Private::RoundAsFrameNumber(TickResolution, StageEndTimeSeconds);
 
 		SubSectionRange = TRange<FFrameNumber>{StartFrame, FMath::Max(StartFrame, StageEndFrame)};
 	}
@@ -1637,8 +1636,8 @@ void FUsdLevelSequenceHelperImpl::CreateSubSequenceSection(ULevelSequence& Seque
 		const double SubStartTimeSeconds = SubLayer.GetStartTimeCode() * SubLayerOffset.Scale / TimeCodesPerSecond;
 		const double SubEndTimeSeconds = SubLayer.GetEndTimeCode() * SubLayerOffset.Scale / TimeCodesPerSecond;
 
-		StartFrame = UsdLevelSequenceHelperImpl::RoundAsFrameNumber(TickResolution, SubStartTimeSeconds);
-		const FFrameNumber EndFrame = UsdLevelSequenceHelperImpl::RoundAsFrameNumber(TickResolution, SubEndTimeSeconds);
+		StartFrame = UE::USDLevelSequenceHelper::Private::RoundAsFrameNumber(TickResolution, SubStartTimeSeconds);
+		const FFrameNumber EndFrame = UE::USDLevelSequenceHelper::Private::RoundAsFrameNumber(TickResolution, SubEndTimeSeconds);
 
 		SubSectionRange = TRange<FFrameNumber>{StartFrame, EndFrame};
 	}
@@ -1764,7 +1763,7 @@ void FUsdLevelSequenceHelperImpl::UpdateSubSectionTimeRanges(bool bShowResizedSe
 
 		FFrameRate TickResolution = MovieScene->GetTickResolution();
 		double TimeCodesPerSecond = GetTimeCodesPerSecond();
-		const FFrameNumber StageEndFrame = UsdLevelSequenceHelperImpl::RoundAsFrameNumber(TickResolution, StageEndTimeSeconds);
+		const FFrameNumber StageEndFrame = UE::USDLevelSequenceHelper::Private::RoundAsFrameNumber(TickResolution, StageEndTimeSeconds);
 
 		FLayerTimeInfo* LayerTimeInfo = LayerTimeInfosByLayerIdentifier.Find(*IdentifierAndSeq.Key);
 
@@ -1948,8 +1947,8 @@ void FUsdLevelSequenceHelperImpl::CreateTimeTrack(const FLayerTimeInfo& Info)
 		const double TimeCodesPerSecond = GetTimeCodesPerSecond();
 
 		FFrameRate DestTickRate = MovieScene->GetTickResolution();
-		FFrameNumber StartFrame = UsdLevelSequenceHelperImpl::RoundAsFrameNumber(DestTickRate, StartTimeCode / TimeCodesPerSecond);
-		FFrameNumber EndFrame = UsdLevelSequenceHelperImpl::RoundAsFrameNumber(DestTickRate, EndTimeCode / TimeCodesPerSecond);
+		FFrameNumber StartFrame = UE::USDLevelSequenceHelper::Private::RoundAsFrameNumber(DestTickRate, StartTimeCode / TimeCodesPerSecond);
+		FFrameNumber EndFrame = UE::USDLevelSequenceHelper::Private::RoundAsFrameNumber(DestTickRate, EndTimeCode / TimeCodesPerSecond);
 
 		TRange<FFrameNumber> PlaybackRange(StartFrame, EndFrame);
 
@@ -2142,7 +2141,7 @@ void FUsdLevelSequenceHelperImpl::AddCommonTracks(const UUsdPrimTwin& PrimTwin,
 
 			if (UMovieScene* MovieScene = AttributeSequence->GetMovieScene())
 			{
-				UsdLevelSequenceHelperImpl::FScopedReadOnlyDisable ReadOnlyGuard{MovieScene, Layer, UsdStage};
+				UE::USDLevelSequenceHelper::Private::FScopedReadOnlyDisable ReadOnlyGuard{MovieScene, Layer, UsdStage};
 
 				if (bNeedTrackToCompensateResetXformOp)
 				{
@@ -2256,7 +2255,7 @@ void FUsdLevelSequenceHelperImpl::AddCommonTracks(const UUsdPrimTwin& PrimTwin,
 
 					if (UMovieScene* MovieScene = AttributeSequence->GetMovieScene())
 					{
-						UsdLevelSequenceHelperImpl::FScopedReadOnlyDisable ReadOnlyGuard{MovieScene, SequenceLayer, UsdStage};
+						UE::USDLevelSequenceHelper::Private::FScopedReadOnlyDisable ReadOnlyGuard{MovieScene, SequenceLayer, UsdStage};
 
 						if (UMovieSceneVisibilityTrack* VisibilityTrack = AddTrack<UMovieSceneVisibilityTrack>(
 								UnrealIdentifiers::HiddenInGamePropertyName,
@@ -2378,7 +2377,7 @@ void FUsdLevelSequenceHelperImpl::AddBoundsTracks(const UUsdPrimTwin& PrimTwin,
 		return;
 	}
 
-	UsdLevelSequenceHelperImpl::FScopedReadOnlyDisable ReadOnlyGuard{MovieScene, TargetLayer, UsdStage};
+	UE::USDLevelSequenceHelper::Private::FScopedReadOnlyDisable ReadOnlyGuard{MovieScene, TargetLayer, UsdStage};
 
 	UMovieSceneDoubleVectorTrack* MinTrack = AddTrack<UMovieSceneDoubleVectorTrack>(
 		GET_MEMBER_NAME_CHECKED(UUsdDrawModeComponent, BoundsMin),
@@ -2423,15 +2422,26 @@ void FUsdLevelSequenceHelperImpl::AddBoundsTracks(const UUsdPrimTwin& PrimTwin,
 	PrimPathByLevelSequenceName.AddUnique(TargetSequence->GetFName(), Prim.GetPrimPath().GetString());
 }
 
-void FUsdLevelSequenceHelperImpl::AddCameraTracks(const UUsdPrimTwin& PrimTwin, const UE::FUsdPrim& Prim)
+namespace UE::USDLevelSequenceHelper::Private
 {
-	TArray<FName> TrackedProperties = {
+	static TSet<FName> TrackedCameraProperties = {
 		UnrealIdentifiers::CurrentFocalLengthPropertyName,
 		UnrealIdentifiers::ManualFocusDistancePropertyName,
 		UnrealIdentifiers::CurrentAperturePropertyName,
 		UnrealIdentifiers::SensorWidthPropertyName,
-		UnrealIdentifiers::SensorHeightPropertyName};
+		UnrealIdentifiers::SensorHeightPropertyName,
+		UnrealIdentifiers::SensorHorizontalOffsetPropertyName,
+		UnrealIdentifiers::SensorVerticalOffsetPropertyName,
+		UnrealIdentifiers::ExposureCompensationPropertyName,
+		UnrealIdentifiers::ProjectionModePropertyName,
+		UnrealIdentifiers::OrthoFarClipPlanePropertyName,
+		UnrealIdentifiers::OrthoNearClipPlanePropertyName,
+		UnrealIdentifiers::CustomNearClipppingPlanePropertyName,
+	};
+}
 
+void FUsdLevelSequenceHelperImpl::AddCameraTracks(const UUsdPrimTwin& PrimTwin, const UE::FUsdPrim& Prim)
+{
 	// For ACineCameraActor the camera component is not the actual root component, so we need to fetch it manually here
 	ACineCameraActor* CameraActor = Cast<ACineCameraActor>(PrimTwin.GetSceneComponent()->GetOwner());
 	if (!CameraActor)
@@ -2444,7 +2454,7 @@ void FUsdLevelSequenceHelperImpl::AddCameraTracks(const UUsdPrimTwin& PrimTwin,
 		return;
 	}
 
-	for (const FName& PropertyName : TrackedProperties)
+	for (const FName& PropertyName : UE::USDLevelSequenceHelper::Private::TrackedCameraProperties)
 	{
 		TArray<UE::FUsdAttribute> Attrs = UsdUtils::GetAttributesForProperty(Prim, PropertyName);
 		if (Attrs.Num() < 1)
@@ -2478,7 +2488,7 @@ void FUsdLevelSequenceHelperImpl::AddCameraTracks(const UUsdPrimTwin& PrimTwin,
 				continue;
 			}
 
-			UsdLevelSequenceHelperImpl::FScopedReadOnlyDisable ReadOnlyGuard{MovieScene, SequenceLayer, UsdStage};
+			UE::USDLevelSequenceHelper::Private::FScopedReadOnlyDisable ReadOnlyGuard{MovieScene, SequenceLayer, UsdStage};
 
 			TArray<double> TimeSamples;
 			if (!Attr.GetTimeSamples(TimeSamples))
@@ -2613,7 +2623,7 @@ void FUsdLevelSequenceHelperImpl::AddLightTracks(const UUsdPrimTwin& PrimTwin, c
 			continue;
 		}
 
-		UsdLevelSequenceHelperImpl::FScopedReadOnlyDisable ReadOnlyGuard{MovieScene, PrimLayer, UsdStage};
+		UE::USDLevelSequenceHelper::Private::FScopedReadOnlyDisable ReadOnlyGuard{MovieScene, PrimLayer, UsdStage};
 
 		UsdToUnreal::FPropertyTrackReader Reader = UsdToUnreal::CreatePropertyTrackReader(Prim, PropertyPath);
 
@@ -2721,7 +2731,7 @@ void FUsdLevelSequenceHelperImpl::AddSkeletalTracks(const UUsdPrimTwin& PrimTwin
 		return;
 	}
 
-	UsdLevelSequenceHelperImpl::FScopedReadOnlyDisable ReadOnlyGuard{MovieScene, SkelAnimationLayer, UsdStage};
+	UE::USDLevelSequenceHelper::Private::FScopedReadOnlyDisable ReadOnlyGuard{MovieScene, SkelAnimationLayer, UsdStage};
 
 	// We will mute all SkelAnimation attributes if we mute, so here let's only consider something muted
 	// if it has all attributes muted as well.
@@ -2797,7 +2807,7 @@ void FUsdLevelSequenceHelperImpl::AddGeometryCacheTracks(const UUsdPrimTwin& Pri
 		return;
 	}
 
-	UsdLevelSequenceHelperImpl::FScopedReadOnlyDisable ReadOnlyGuard{MovieScene, GeometryCacheLayer, UsdStage};
+	UE::USDLevelSequenceHelper::Private::FScopedReadOnlyDisable ReadOnlyGuard{MovieScene, GeometryCacheLayer, UsdStage};
 
 	const bool bIsMuted = false;
 	if (UMovieSceneGeometryCacheTrack* GeometryCacheTrack = AddTrack<
@@ -2873,7 +2883,7 @@ void FUsdLevelSequenceHelperImpl::AddGroomTracks(const UUsdPrimTwin& PrimTwin, c
 		return;
 	}
 
-	UsdLevelSequenceHelperImpl::FScopedReadOnlyDisable ReadOnlyGuard{MovieScene, GroomLayer, UsdStage};
+	UE::USDLevelSequenceHelper::Private::FScopedReadOnlyDisable ReadOnlyGuard{MovieScene, GroomLayer, UsdStage};
 
 	const bool bIsMuted = false;
 	if (UMovieSceneGroomCacheTrack* GroomCacheTrack = AddTrack<
@@ -2912,7 +2922,7 @@ void FUsdLevelSequenceHelperImpl::AddVolumeTracks(const UUsdPrimTwin& PrimTwin,
 		return;
 	}
 
-	UsdLevelSequenceHelperImpl::FScopedReadOnlyDisable ReadOnlyGuard{MovieScene, PrimLayer, UsdStage};
+	UE::USDLevelSequenceHelper::Private::FScopedReadOnlyDisable ReadOnlyGuard{MovieScene, PrimLayer, UsdStage};
 
 	const FName PropertyPath = GET_MEMBER_NAME_CHECKED(UHeterogeneousVolumeComponent, Frame);
 
@@ -3060,7 +3070,7 @@ void FUsdLevelSequenceHelperImpl::AddVolumeTracks(const UUsdPrimTwin& PrimTwin,
 	PrimPathByLevelSequenceName.AddUnique(PrimSequence->GetFName(), Prim.GetPrimPath().GetString());
 }
 
-namespace UE::LevelSequenceHelper::Private
+namespace UE::USDLevelSequenceHelper::Private
 {
 	template<typename T>
 	TOptional<T> GetAuthoredValue(const UE::FUsdPrim& Prim, const FString& AttrName)
@@ -3081,7 +3091,7 @@ namespace UE::LevelSequenceHelper::Private
 void FUsdLevelSequenceHelperImpl::AddAudioTracks(const UUsdPrimTwin& PrimTwin, const UE::FUsdPrim& Prim)
 {
 	using namespace UnrealIdentifiers;
-	using namespace UE::LevelSequenceHelper::Private;
+	using namespace UE::USDLevelSequenceHelper::Private;
 
 	UAudioComponent* AudioComponent = Cast<UAudioComponent>(PrimTwin.GetSceneComponent());
 	if (!AudioComponent)
@@ -3135,7 +3145,7 @@ void FUsdLevelSequenceHelperImpl::AddAudioTracks(const UUsdPrimTwin& PrimTwin, c
 			return;
 		}
 
-		UsdLevelSequenceHelperImpl::FScopedReadOnlyDisable ReadOnlyGuard{MovieScene, SequenceLayer, UsdStage};
+		UE::USDLevelSequenceHelper::Private::FScopedReadOnlyDisable ReadOnlyGuard{MovieScene, SequenceLayer, UsdStage};
 
 		// Since the schema has fallbacks for these, we should always have values here.
 		//
@@ -3506,7 +3516,7 @@ TrackType* FUsdLevelSequenceHelperImpl::AddTrack(
 #endif	  // WITH_EDITOR
 	}
 
-	UsdLevelSequenceHelperImpl::MuteTrack(Track, MovieScene, ComponentBinding.ToString(), Track->GetName(), bIsMuted);
+	UE::USDLevelSequenceHelper::Private::MuteTrack(Track, MovieScene, ComponentBinding.ToString(), Track->GetName(), bIsMuted);
 
 	return Track;
 }
@@ -3732,7 +3742,7 @@ void FUsdLevelSequenceHelperImpl::UpdateControlRigTracks(UUsdPrimTwin& PrimTwin)
 		{
 			UAnimSeqExportOption* NewOptions = NewObject<UAnimSeqExportOption>();
 
-			UsdLevelSequenceHelperImpl::BakeToControlRig(
+			UE::USDLevelSequenceHelper::Private::BakeToControlRig(
 				ComponentToBind->GetWorld(),
 				SkelAnimationSequence,
 				ControlRigClass,
@@ -3860,7 +3870,7 @@ void FUsdLevelSequenceHelperImpl::RefreshSequencer()
 		return;
 	}
 
-	if (TSharedPtr<ISequencer> Sequencer = UsdLevelSequenceHelperImpl::GetOpenedSequencerForLevelSequence(MainLevelSequence))
+	if (TSharedPtr<ISequencer> Sequencer = UE::USDLevelSequenceHelper::Private::GetOpenedSequencerForLevelSequence(MainLevelSequence))
 	{
 		// Don't try refreshing the sequencer if its displaying a stale sequence (e.g. during busy transitions like import) as it
 		// can crash
@@ -4004,8 +4014,8 @@ void FUsdLevelSequenceHelperImpl::UpdateMovieSceneTimeRanges(UMovieScene& MovieS
 		}
 
 		const FFrameRate TickResolution = MovieScene.GetTickResolution();
-		const FFrameNumber StartFrame = UsdLevelSequenceHelperImpl::RoundAsFrameNumber(TickResolution, StartTimeCode / TimeCodesPerSecond);
-		const FFrameNumber EndFrame = UsdLevelSequenceHelperImpl::RoundAsFrameNumber(TickResolution, EndTimeCode / TimeCodesPerSecond);
+		const FFrameNumber StartFrame = UE::USDLevelSequenceHelper::Private::RoundAsFrameNumber(TickResolution, StartTimeCode / TimeCodesPerSecond);
+		const FFrameNumber EndFrame = UE::USDLevelSequenceHelper::Private::RoundAsFrameNumber(TickResolution, EndTimeCode / TimeCodesPerSecond);
 		TRange<FFrameNumber> TimeRange = TRange<FFrameNumber>::Inclusive(StartFrame, EndFrame);
 
 		MovieScene.SetPlaybackRange(TimeRange);
@@ -4340,7 +4350,7 @@ FGuid FUsdLevelSequenceHelperImpl::GetOrCreateComponentBinding(
 
 void FUsdLevelSequenceHelperImpl::HandleMovieSceneChange(UMovieScene& MovieScene)
 {
-	using namespace UsdLevelSequenceHelperImpl;
+	using namespace UE::USDLevelSequenceHelper::Private;
 
 	// It's possible to get this called when the actor and it's level sequences are being all destroyed in one go.
 	// We need the FScopedBlockNotices in this function, but if our StageActor is already being destroyed, we can't reliably
@@ -4440,7 +4450,7 @@ void FUsdLevelSequenceHelperImpl::HandleMovieSceneChange(UMovieScene& MovieScene
 	auto RemoveTimeSamplesForPropertyIfNeeded =
 		[&MovieScene, &RemoveTimeSamplesForAttr](const UE::FUsdPrim& Prim, const FGuid& Guid, const FName& PropertyPath)
 	{
-		if (!UsdLevelSequenceHelperImpl::FindTrackTypeOrDerived<UMovieScenePropertyTrack>(&MovieScene, Guid, PropertyPath))
+		if (!UE::USDLevelSequenceHelper::Private::FindTrackTypeOrDerived<UMovieScenePropertyTrack>(&MovieScene, Guid, PropertyPath))
 		{
 			for (UE::FUsdAttribute& Attr : UsdUtils::GetAttributesForProperty(Prim, PropertyPath))
 			{
@@ -4508,6 +4518,13 @@ void FUsdLevelSequenceHelperImpl::HandleMovieSceneChange(UMovieScene& MovieScene
 						RemoveTimeSamplesForPropertyIfNeeded(UsdPrim, Guid, UnrealIdentifiers::CurrentAperturePropertyName);
 						RemoveTimeSamplesForPropertyIfNeeded(UsdPrim, Guid, UnrealIdentifiers::SensorWidthPropertyName);
 						RemoveTimeSamplesForPropertyIfNeeded(UsdPrim, Guid, UnrealIdentifiers::SensorHeightPropertyName);
+						RemoveTimeSamplesForPropertyIfNeeded(UsdPrim, Guid, UnrealIdentifiers::SensorHorizontalOffsetPropertyName);
+						RemoveTimeSamplesForPropertyIfNeeded(UsdPrim, Guid, UnrealIdentifiers::SensorVerticalOffsetPropertyName);
+						RemoveTimeSamplesForPropertyIfNeeded(UsdPrim, Guid, UnrealIdentifiers::ExposureCompensationPropertyName);
+						RemoveTimeSamplesForPropertyIfNeeded(UsdPrim, Guid, UnrealIdentifiers::ProjectionModePropertyName);
+						RemoveTimeSamplesForPropertyIfNeeded(UsdPrim, Guid, UnrealIdentifiers::OrthoFarClipPlanePropertyName);
+						RemoveTimeSamplesForPropertyIfNeeded(UsdPrim, Guid, UnrealIdentifiers::OrthoNearClipPlanePropertyName);
+						RemoveTimeSamplesForPropertyIfNeeded(UsdPrim, Guid, UnrealIdentifiers::CustomNearClipppingPlanePropertyName);
 					}
 					else if (bIsLight)
 					{
@@ -4609,7 +4626,8 @@ void FUsdLevelSequenceHelperImpl::HandleControlRigSectionChange(UMovieSceneContr
 		return;
 	}
 
-	USkeletalMeshComponent* BoundComponent = Cast<USkeletalMeshComponent>(UsdLevelSequenceHelperImpl::LocateBoundObject(*LevelSequence, *Possessable)
+	USkeletalMeshComponent* BoundComponent = Cast<USkeletalMeshComponent>(
+		UE::USDLevelSequenceHelper::Private::LocateBoundObject(*LevelSequence, *Possessable)
 	);
 	if (!BoundComponent)
 	{
@@ -4658,7 +4676,7 @@ void FUsdLevelSequenceHelperImpl::HandleControlRigSectionChange(UMovieSceneContr
 	}
 	UE::FUsdEditContext EditContext{UsdStage, EditTargetLayer};
 
-	TSharedPtr<ISequencer> PinnedSequencer = UsdLevelSequenceHelperImpl::GetOpenedSequencerForLevelSequence(MainLevelSequence);
+	TSharedPtr<ISequencer> PinnedSequencer = UE::USDLevelSequenceHelper::Private::GetOpenedSequencerForLevelSequence(MainLevelSequence);
 
 	// Fetch a sequence player we can use. We'll almost always have the sequencer opened here (we are responding to a transaction
 	// where the section was changed after all), but its possible to have a fallback too
@@ -4788,7 +4806,7 @@ void FUsdLevelSequenceHelperImpl::HandleTrackChange(UMovieSceneTrack& Track, boo
 		return;
 	}
 
-	UObject* BoundObject = UsdLevelSequenceHelperImpl::LocateBoundObject(*Sequence, *Possessable);
+	UObject* BoundObject = UE::USDLevelSequenceHelper::Private::LocateBoundObject(*Sequence, *Possessable);
 	if (!BoundObject)
 	{
 		return;
@@ -4809,7 +4827,7 @@ void FUsdLevelSequenceHelperImpl::HandleTrackChange(UMovieSceneTrack& Track, boo
 		else
 		{
 #if WITH_EDITOR
-			UsdLevelSequenceHelperImpl::ShowStageActorPropertyTrackWarning(PropertyPath);
+			UE::USDLevelSequenceHelper::Private::ShowStageActorPropertyTrackWarning(PropertyPath);
 #endif	  // WITH_EDITOR
 			MovieScene->RemoveTrack(*PropertyTrack);
 		}
@@ -4847,11 +4865,7 @@ void FUsdLevelSequenceHelperImpl::HandleTrackChange(UMovieSceneTrack& Track, boo
 			// In the scenario where we're trying to make non-decomposed Camera prims work, we only ever want to write out
 			// actual camera properties from the CameraComponent to the Camera prim. We won't write its USceneComponent
 			// properties, as we will use the ones from the ACineCameraActor's parent USceneComponent instead
-			if (PropertyPath == UnrealIdentifiers::CurrentFocalLengthPropertyName
-				|| PropertyPath == UnrealIdentifiers::ManualFocusDistancePropertyName
-				|| PropertyPath == UnrealIdentifiers::ManualFocusDistancePropertyName
-				|| PropertyPath == UnrealIdentifiers::CurrentAperturePropertyName || PropertyPath == UnrealIdentifiers::SensorWidthPropertyName
-				|| PropertyPath == UnrealIdentifiers::SensorHeightPropertyName)
+			if (UE::USDLevelSequenceHelper::Private::TrackedCameraProperties.Contains(PropertyPath))
 			{
 				PrimTwin = StageActor->RootUsdTwin->Find(BoundSceneComponent->GetAttachParent());
 			}
@@ -4859,7 +4873,7 @@ void FUsdLevelSequenceHelperImpl::HandleTrackChange(UMovieSceneTrack& Track, boo
 			else if (PropertyPath == UnrealIdentifiers::TransformPropertyName)
 			{
 				// Let the user know that we currently don't support transform tracks directly on camera components
-				UsdLevelSequenceHelperImpl::ShowTransformTrackOnCameraComponentWarning(BoundSceneComponent);
+				UE::USDLevelSequenceHelper::Private::ShowTransformTrackOnCameraComponentWarning(BoundSceneComponent);
 			}
 #endif	  // WITH_EDITOR
 		}
@@ -4979,7 +4993,7 @@ void FUsdLevelSequenceHelperImpl::HandleTrackChange(UMovieSceneTrack& Track, boo
 			if (PropertyTrack)
 			{
 #if WITH_EDITOR
-				UsdLevelSequenceHelperImpl::ShowVisibilityWarningIfNeeded(PropertyTrack, UsdPrim);
+				UE::USDLevelSequenceHelper::Private::ShowVisibilityWarningIfNeeded(PropertyTrack, UsdPrim);
 #endif
 
 				TSet<FName> PropertyPathsToRefresh;
@@ -5210,13 +5224,13 @@ public:
 	{
 		return false;
 	}
-	void SetPrimLinkCache(UUsdPrimLinkCache* PrimLinkCache){};
-	void SetBBoxCache(TSharedPtr<UE::FUsdGeomBBoxCache> InBBoxCache){};
+	void SetPrimLinkCache(UUsdPrimLinkCache* PrimLinkCache) {};
+	void SetBBoxCache(TSharedPtr<UE::FUsdGeomBBoxCache> InBBoxCache) {};
 	bool HasData() const
 	{
 		return false;
 	};
-	void Clear(){};
+	void Clear() {};
 
 	void CreateLocalLayersSequences()
 	{
@@ -5232,8 +5246,8 @@ public:
 	{
 		return EUsdRootMotionHandling::NoAdditionalRootMotion;
 	}
-	void SetRootMotionHandling(EUsdRootMotionHandling NewValue){};
-	void OnStageActorRenamed(){};
+	void SetRootMotionHandling(EUsdRootMotionHandling NewValue) {};
+	void OnStageActorRenamed() {};
 
 	void AddPrim(UUsdPrimTwin& PrimTwin, bool bForceVisibilityTracks, TOptional<bool> HasAnimatedBounds)
 	{
diff --git a/Engine/Plugins/Interchange/Runtime/Source/Import/Private/Usd/InterchangeUsdTranslator.cpp b/Engine/Plugins/Interchange/Runtime/Source/Import/Private/Usd/InterchangeUsdTranslator.cpp
--- a/Engine/Plugins/Interchange/Runtime/Source/Import/Private/Usd/InterchangeUsdTranslator.cpp
+++ b/Engine/Plugins/Interchange/Runtime/Source/Import/Private/Usd/InterchangeUsdTranslator.cpp
@@ -109,23 +109,29 @@ namespace UE::InterchangeUsdTranslator::Private
 		{UnrealIdentifiers::HiddenInGamePropertyName, 			EInterchangePropertyTracks::Visibility}, // Binding visibility to the actor works better for cameras
 
 		// Camera properties
-		{UnrealIdentifiers::CurrentFocalLengthPropertyName, 	EInterchangePropertyTracks::CameraCurrentFocalLength},
-		{UnrealIdentifiers::ManualFocusDistancePropertyName, 	EInterchangePropertyTracks::CameraFocusSettingsManualFocusDistance},
-		{UnrealIdentifiers::CurrentAperturePropertyName, 		EInterchangePropertyTracks::CameraCurrentAperture},
-		{UnrealIdentifiers::SensorWidthPropertyName, 			EInterchangePropertyTracks::CameraFilmbackSensorWidth},
-		{UnrealIdentifiers::SensorHeightPropertyName, 			EInterchangePropertyTracks::CameraFilmbackSensorHeight},
+		// TODO: Need to add support for SensorHorizontalOffset and SensorVerticalOffset, once Interchange supports those
+		{UnrealIdentifiers::CurrentFocalLengthPropertyName, 		EInterchangePropertyTracks::CameraCurrentFocalLength},
+		{UnrealIdentifiers::ManualFocusDistancePropertyName, 		EInterchangePropertyTracks::CameraFocusSettingsManualFocusDistance},
+		{UnrealIdentifiers::CurrentAperturePropertyName, 			EInterchangePropertyTracks::CameraCurrentAperture},
+		{UnrealIdentifiers::SensorWidthPropertyName, 				EInterchangePropertyTracks::CameraFilmbackSensorWidth},
+		{UnrealIdentifiers::SensorHeightPropertyName, 				EInterchangePropertyTracks::CameraFilmbackSensorHeight},
+		{UnrealIdentifiers::ExposureCompensationPropertyName, 		EInterchangePropertyTracks::CameraPostProcessSettingsAutoExposureBias},
+		{UnrealIdentifiers::ProjectionModePropertyName, 			EInterchangePropertyTracks::CameraProjectionMode},
+		{UnrealIdentifiers::OrthoFarClipPlanePropertyName, 			EInterchangePropertyTracks::CameraOrthoFarClipPlane},
+		{UnrealIdentifiers::OrthoNearClipPlanePropertyName, 		EInterchangePropertyTracks::CameraOrthoNearClipPlane},
+		{UnrealIdentifiers::CustomNearClipppingPlanePropertyName,	EInterchangePropertyTracks::CameraCustomNearClippingPlane},
 
 		// Light properties
-		{UnrealIdentifiers::LightColorPropertyName, 			EInterchangePropertyTracks::LightColor},
-		{UnrealIdentifiers::TemperaturePropertyName, 			EInterchangePropertyTracks::LightTemperature},
-		{UnrealIdentifiers::UseTemperaturePropertyName, 		EInterchangePropertyTracks::LightUseTemperature},
-		{UnrealIdentifiers::SourceHeightPropertyName, 			EInterchangePropertyTracks::LightSourceHeight},
-		{UnrealIdentifiers::SourceWidthPropertyName, 			EInterchangePropertyTracks::LightSourceWidth},
-		{UnrealIdentifiers::SourceRadiusPropertyName, 			EInterchangePropertyTracks::LightSourceRadius},
-		{UnrealIdentifiers::OuterConeAnglePropertyName, 		EInterchangePropertyTracks::LightOuterConeAngle},
-		{UnrealIdentifiers::InnerConeAnglePropertyName, 		EInterchangePropertyTracks::LightInnerConeAngle},
-		{UnrealIdentifiers::LightSourceAnglePropertyName, 		EInterchangePropertyTracks::LightSourceAngle},
-		{UnrealIdentifiers::IntensityPropertyName, 				EInterchangePropertyTracks::LightIntensity},
+		{UnrealIdentifiers::LightColorPropertyName, 		EInterchangePropertyTracks::LightColor},
+		{UnrealIdentifiers::TemperaturePropertyName, 		EInterchangePropertyTracks::LightTemperature},
+		{UnrealIdentifiers::UseTemperaturePropertyName, 	EInterchangePropertyTracks::LightUseTemperature},
+		{UnrealIdentifiers::SourceHeightPropertyName, 		EInterchangePropertyTracks::LightSourceHeight},
+		{UnrealIdentifiers::SourceWidthPropertyName, 		EInterchangePropertyTracks::LightSourceWidth},
+		{UnrealIdentifiers::SourceRadiusPropertyName, 		EInterchangePropertyTracks::LightSourceRadius},
+		{UnrealIdentifiers::OuterConeAnglePropertyName, 	EInterchangePropertyTracks::LightOuterConeAngle},
+		{UnrealIdentifiers::InnerConeAnglePropertyName, 	EInterchangePropertyTracks::LightInnerConeAngle},
+		{UnrealIdentifiers::LightSourceAnglePropertyName, 	EInterchangePropertyTracks::LightSourceAngle},
+		{UnrealIdentifiers::IntensityPropertyName, 			EInterchangePropertyTracks::LightIntensity},
 	};
 	// clang-format on
 
@@ -657,7 +663,8 @@ namespace UE::InterchangeUsdTranslator::Private
 				Value.UVScale.GetVector()[0],
 				Value.UVScale.GetVector()[1],
 				Value.UVTranslation[0],
-				Value.UVTranslation[1]};
+				Value.UVTranslation[1]
+			};
 			MaterialNode.AddVectorParameterValue(*FString::Printf(TEXT("%sScaleTranslation"), **ParameterName), ScaleAndTranslation);
 			MaterialNode.AddScalarParameterValue(*FString::Printf(TEXT("%sRotation"), **ParameterName), Value.UVRotation);
 
@@ -677,10 +684,11 @@ namespace UE::InterchangeUsdTranslator::Private
 					**ParameterName,
 					*Prim.GetPrimPath().GetString(),
 					*UsdUtils::StringifyMap(PrimvarToUVIndex),
-					Value.Primvar.IsEmpty() ? TEXT(
-						" Is your UsdUVTexture Shader missing the 'inputs:st' attribute? (It specifies which UV set to sample the texture with)"
-					)
-											: TEXT("")
+					Value.Primvar.IsEmpty()
+						? TEXT(
+							  " Is your UsdUVTexture Shader missing the 'inputs:st' attribute? (It specifies which UV set to sample the texture with)"
+						  )
+						: TEXT("")
 				);
 			}
 
diff --git a/Engine/Plugins/Runtime/USDCore/Source/USDUtilities/Private/USDConversionUtils.cpp b/Engine/Plugins/Runtime/USDCore/Source/USDUtilities/Private/USDConversionUtils.cpp
--- a/Engine/Plugins/Runtime/USDCore/Source/USDUtilities/Private/USDConversionUtils.cpp
+++ b/Engine/Plugins/Runtime/USDCore/Source/USDUtilities/Private/USDConversionUtils.cpp
@@ -356,6 +356,7 @@ template USDUTILITIES_API bool UsdUtils::GetUsdValue<bool>(const pxr::UsdAttribu
 template USDUTILITIES_API float UsdUtils::GetUsdValue<float>(const pxr::UsdAttribute& Attribute, pxr::UsdTimeCode TimeCode);
 template USDUTILITIES_API double UsdUtils::GetUsdValue<double>(const pxr::UsdAttribute& Attribute, pxr::UsdTimeCode TimeCode);
 template USDUTILITIES_API pxr::TfToken UsdUtils::GetUsdValue<pxr::TfToken>(const pxr::UsdAttribute& Attribute, pxr::UsdTimeCode TimeCode);
+template USDUTILITIES_API pxr::GfVec2f UsdUtils::GetUsdValue<pxr::GfVec2f>(const pxr::UsdAttribute& Attribute, pxr::UsdTimeCode TimeCode);
 template USDUTILITIES_API pxr::GfVec3f UsdUtils::GetUsdValue<pxr::GfVec3f>(const pxr::UsdAttribute& Attribute, pxr::UsdTimeCode TimeCode);
 template USDUTILITIES_API pxr::GfMatrix4d UsdUtils::GetUsdValue<pxr::GfMatrix4d>(const pxr::UsdAttribute& Attribute, pxr::UsdTimeCode TimeCode);
 template USDUTILITIES_API pxr::SdfAssetPath UsdUtils::GetUsdValue<pxr::SdfAssetPath>(const pxr::UsdAttribute& Attribute, pxr::UsdTimeCode TimeCode);
diff --git a/Engine/Plugins/Runtime/USDCore/Source/USDUtilities/Private/USDPrimConversion.cpp b/Engine/Plugins/Runtime/USDCore/Source/USDUtilities/Private/USDPrimConversion.cpp
--- a/Engine/Plugins/Runtime/USDCore/Source/USDUtilities/Private/USDPrimConversion.cpp
+++ b/Engine/Plugins/Runtime/USDCore/Source/USDUtilities/Private/USDPrimConversion.cpp
@@ -509,31 +509,79 @@ bool UsdToUnreal::ConvertGeomCamera(const UE::FUsdPrim& Prim, UCineCameraCompone
 	UE::FUsdStage Stage = Prim.GetStage();
 	FUsdStageInfo StageInfo(Stage);
 
-	CameraComponent.CurrentFocalLength = UsdToUnreal::ConvertDistance(
-		StageInfo,
-		UsdUtils::GetUsdValue<float>(GeomCamera.GetFocalLengthAttr(), UsdTimeCode)
-	);
+	if (pxr::UsdAttribute Attr = GeomCamera.GetFocalLengthAttr(); Attr.HasAuthoredValue())
+	{
+		CameraComponent.SetCurrentFocalLength(UsdToUnreal::ConvertDistance(StageInfo, UsdUtils::GetUsdValue<float>(Attr, UsdTimeCode)));
+	}
 
-	CameraComponent.FocusSettings.ManualFocusDistance = UsdToUnreal::ConvertDistance(
-		StageInfo,
-		UsdUtils::GetUsdValue<float>(GeomCamera.GetFocusDistanceAttr(), UsdTimeCode)
-	);
+	if (pxr::UsdAttribute Attr = GeomCamera.GetFocusDistanceAttr(); Attr.HasAuthoredValue())
+	{
+		CameraComponent.FocusSettings.ManualFocusDistance = UsdToUnreal::ConvertDistance(
+			StageInfo,
+			UsdUtils::GetUsdValue<float>(GeomCamera.GetFocusDistanceAttr(), UsdTimeCode)
+		);
 
-	if (FMath::IsNearlyZero(CameraComponent.FocusSettings.ManualFocusDistance))
+		if (FMath::IsNearlyZero(CameraComponent.FocusSettings.ManualFocusDistance))
+		{
+			CameraComponent.FocusSettings.FocusMethod = ECameraFocusMethod::DoNotOverride;
+		}
+		else
+		{
+			CameraComponent.FocusSettings.FocusMethod = ECameraFocusMethod::Manual;
+		}
+	}
+
+	if (pxr::UsdAttribute Attr = GeomCamera.GetFStopAttr(); Attr.HasAuthoredValue())
 	{
-		CameraComponent.FocusSettings.FocusMethod = ECameraFocusMethod::DoNotOverride;
+		CameraComponent.SetCurrentAperture(UsdUtils::GetUsdValue<float>(Attr, UsdTimeCode));
 	}
 
-	CameraComponent.CurrentAperture = UsdUtils::GetUsdValue<float>(GeomCamera.GetFStopAttr(), UsdTimeCode);
+	if (pxr::UsdAttribute Attr = GeomCamera.GetHorizontalApertureAttr(); Attr.HasAuthoredValue())
+	{
+		CameraComponent.Filmback.SensorWidth = UsdToUnreal::ConvertDistance(StageInfo, UsdUtils::GetUsdValue<float>(Attr, UsdTimeCode));
+	}
 
-	CameraComponent.Filmback.SensorWidth = UsdToUnreal::ConvertDistance(
-		StageInfo,
-		UsdUtils::GetUsdValue<float>(GeomCamera.GetHorizontalApertureAttr(), UsdTimeCode)
-	);
-	CameraComponent.Filmback.SensorHeight = UsdToUnreal::ConvertDistance(
-		StageInfo,
-		UsdUtils::GetUsdValue<float>(GeomCamera.GetVerticalApertureAttr(), UsdTimeCode)
-	);
+	if (pxr::UsdAttribute Attr = GeomCamera.GetVerticalApertureAttr(); Attr.HasAuthoredValue())
+	{
+		CameraComponent.Filmback.SensorHeight = UsdToUnreal::ConvertDistance(StageInfo, UsdUtils::GetUsdValue<float>(Attr, UsdTimeCode));
+	}
+
+	if (pxr::UsdAttribute Attr = GeomCamera.GetHorizontalApertureOffsetAttr(); Attr.HasAuthoredValue())
+	{
+		CameraComponent.Filmback.SensorHorizontalOffset = UsdToUnreal::ConvertDistance(StageInfo, UsdUtils::GetUsdValue<float>(Attr, UsdTimeCode));
+	}
+
+	if (pxr::UsdAttribute Attr = GeomCamera.GetVerticalApertureOffsetAttr(); Attr.HasAuthoredValue())
+	{
+		CameraComponent.Filmback.SensorVerticalOffset = UsdToUnreal::ConvertDistance(StageInfo, UsdUtils::GetUsdValue<float>(Attr, UsdTimeCode));
+	}
+
+	if (pxr::UsdAttribute Attr = GeomCamera.GetExposureAttr(); Attr.HasAuthoredValue())
+	{
+		CameraComponent.PostProcessSettings.bOverride_AutoExposureBias = true;
+		CameraComponent.PostProcessSettings.AutoExposureBias = UsdUtils::GetUsdValue<float>(Attr, UsdTimeCode);
+	}
+
+	if (pxr::UsdAttribute Attr = GeomCamera.GetProjectionAttr(); Attr.HasAuthoredValue())
+	{
+		CameraComponent.SetProjectionMode(
+			UsdUtils::GetUsdValue<pxr::TfToken>(Attr, UsdTimeCode) == pxr::UsdGeomTokens->orthographic ? ECameraProjectionMode::Orthographic
+																									   : ECameraProjectionMode::Perspective
+		);
+	}
+
+	if (pxr::UsdAttribute Attr = GeomCamera.GetClippingRangeAttr(); Attr.HasAuthoredValue())
+	{
+		pxr::GfVec2f NearAndFarSceneUnits = UsdUtils::GetUsdValue<pxr::GfVec2f>(Attr, UsdTimeCode);
+		float UENear = UsdToUnreal::ConvertDistance(StageInfo, NearAndFarSceneUnits[0]);
+		float UEFar = UsdToUnreal::ConvertDistance(StageInfo, NearAndFarSceneUnits[1]);
+
+		CameraComponent.SetOrthoNearClipPlane(UENear);
+		CameraComponent.SetOrthoFarClipPlane(UEFar);
+
+		CameraComponent.SetCustomNearClippingPlane(UENear);
+		CameraComponent.bOverride_CustomNearClippingPlane = true;
+	}
 
 	return true;
 }
@@ -1127,7 +1175,7 @@ bool UsdToUnreal::ConvertBoundsTimeSamples(
 			FVector UESpaceUSDMax = UsdToUnreal::ConvertVector(StageInfo, Box.GetMax());
 			UEBox = FBox{
 				TArray<FVector>{UESpaceUSDMin, UESpaceUSDMax}
-			 };
+			};
 		}
 
 		MinXValues.Emplace_GetRef(UEBox.Min.X).InterpMode = InterpMode;
@@ -1224,6 +1272,7 @@ UsdToUnreal::FPropertyTrackReader UsdToUnreal::CreatePropertyTrackReader(
 
 	if (pxr::UsdGeomCamera Camera{UsdPrim})
 	{
+		TOptional<int32> ClippingPlaneIndex;
 		bool bConvertDistance = true;
 		pxr::UsdAttribute Attr;
 
@@ -1248,10 +1297,60 @@ UsdToUnreal::FPropertyTrackReader UsdToUnreal::CreatePropertyTrackReader(
 		{
 			Attr = Camera.GetVerticalApertureAttr();
 		}
+		else if (PropertyPath == UnrealIdentifiers::SensorHorizontalOffsetPropertyName)
+		{
+			Attr = Camera.GetHorizontalApertureOffsetAttr();
+		}
+		else if (PropertyPath == UnrealIdentifiers::SensorVerticalOffsetPropertyName)
+		{
+			Attr = Camera.GetVerticalApertureOffsetAttr();
+		}
+		else if (PropertyPath == UnrealIdentifiers::ExposureCompensationPropertyName)
+		{
+			Attr = Camera.GetExposureAttr();
+			bConvertDistance = false;
+		}
+		else if (PropertyPath == UnrealIdentifiers::OrthoFarClipPlanePropertyName)
+		{
+			Attr = Camera.GetClippingRangeAttr();
+			ClippingPlaneIndex = 1;	   // Far plane
+		}
+		else if (PropertyPath == UnrealIdentifiers::OrthoNearClipPlanePropertyName)
+		{
+			Attr = Camera.GetClippingRangeAttr();
+			ClippingPlaneIndex = 0;	   // Near plane
+		}
+		else if (PropertyPath == UnrealIdentifiers::CustomNearClipppingPlanePropertyName)
+		{
+			Attr = Camera.GetClippingRangeAttr();
+			ClippingPlaneIndex = 0;	   // Near plane
+		}
+
+		// TODO: Projection mode track reader? Would need new enum flag reader types
 
 		if (Attr)
 		{
-			if (bConvertDistance)
+			if (ClippingPlaneIndex.IsSet() && ClippingPlaneIndex.GetValue() < 2)
+			{
+				pxr::GfVec2f DefaultVec;
+				Attr.Get<pxr::GfVec2f>(&DefaultVec);
+				float DefaultValue = UsdToUnreal::ConvertDistance(StageInfo, DefaultVec[ClippingPlaneIndex.GetValue()]);
+
+				Reader.FloatReader = [Attr, DefaultValue, StageInfo, Index = ClippingPlaneIndex.GetValue()](double UsdTimeCode)
+				{
+					float Result = DefaultValue;
+
+					pxr::GfVec2f DefaultVec;
+					if (Attr.Get<pxr::GfVec2f>(&DefaultVec, UsdTimeCode))
+					{
+						Result = UsdToUnreal::ConvertDistance(StageInfo, DefaultVec[Index]);
+					}
+
+					return Result;
+				};
+				return Reader;
+			}
+			else if (bConvertDistance)
 			{
 				float Default = 0.0;
 				Attr.Get<float>(&Default);
@@ -1821,7 +1920,7 @@ bool UsdToUnreal::ConvertDrawMode(
 		FVector UESpaceUSDMax = UsdToUnreal::ConvertVector(StageInfo, Box.GetMax());
 		FBox UEBox{
 			TArray<FVector>{UESpaceUSDMin, UESpaceUSDMax}
-		 };
+		};
 
 		DrawModeComponent->SetBoundsMin(UEBox.Min);
 		DrawModeComponent->SetBoundsMax(UEBox.Max);
@@ -2135,6 +2234,65 @@ bool UnrealToUsd::ConvertCameraComponent(const UCineCameraComponent& CameraCompo
 		UsdUtils::NotifyIfOverriddenOpinion(Attr);
 	}
 
+	if (pxr::UsdAttribute Attr = GeomCamera.CreateHorizontalApertureOffsetAttr())
+	{
+		Attr.Set<float>(UnrealToUsd::ConvertDistance(StageInfo, CameraComponent.Filmback.SensorHorizontalOffset), UsdTimeCode);
+		UsdUtils::NotifyIfOverriddenOpinion(Attr);
+	}
+
+	if (pxr::UsdAttribute Attr = GeomCamera.CreateVerticalApertureOffsetAttr())
+	{
+		Attr.Set<float>(UnrealToUsd::ConvertDistance(StageInfo, CameraComponent.Filmback.SensorVerticalOffset), UsdTimeCode);
+		UsdUtils::NotifyIfOverriddenOpinion(Attr);
+	}
+
+	if (pxr::UsdAttribute Attr = GeomCamera.CreateExposureAttr(); CameraComponent.PostProcessSettings.bOverride_AutoExposureBias)
+	{
+		Attr.Set<float>(CameraComponent.PostProcessSettings.AutoExposureBias, UsdTimeCode);
+		UsdUtils::NotifyIfOverriddenOpinion(Attr);
+	}
+
+	if (pxr::UsdAttribute Attr = GeomCamera.CreateProjectionAttr())
+	{
+		Attr.Set<pxr::TfToken>(
+			CameraComponent.ProjectionMode == ECameraProjectionMode::Orthographic ? pxr::UsdGeomTokens->orthographic
+																				  : pxr::UsdGeomTokens->perspective,
+			UsdTimeCode
+		);
+		UsdUtils::NotifyIfOverriddenOpinion(Attr);
+	}
+
+	// Write clipping planes from the orthographic mode properties, if we're in orthographic mode
+	if (CameraComponent.ProjectionMode == ECameraProjectionMode::Orthographic)
+	{
+		if (pxr::UsdAttribute Attr = GeomCamera.CreateClippingRangeAttr())
+		{
+			Attr.Set<pxr::GfVec2f>(
+				pxr::GfVec2f{
+					UnrealToUsd::ConvertDistance(StageInfo, CameraComponent.OrthoNearClipPlane),
+					UnrealToUsd::ConvertDistance(StageInfo, CameraComponent.OrthoFarClipPlane)
+				},
+				UsdTimeCode
+			);
+			UsdUtils::NotifyIfOverriddenOpinion(Attr);
+		}
+	}
+	else
+	{
+		if (pxr::UsdAttribute Attr = GeomCamera.CreateClippingRangeAttr(); CameraComponent.bOverride_CustomNearClippingPlane)
+		{
+			// We only have a property for near, so let's leave the far plane at its current value
+			pxr::GfVec2f ClippingRangeUsdUnits{1, 1000000};	   // Values copied from the schema, as a fallback
+			Attr.Get(&ClippingRangeUsdUnits, UsdTimeCode);
+
+			ClippingRangeUsdUnits[0] = UnrealToUsd::ConvertDistance(StageInfo, CameraComponent.CustomNearClippingPlane);
+
+			Attr.Set<pxr::GfVec2f>(ClippingRangeUsdUnits, UsdTimeCode);
+
+			UsdUtils::NotifyIfOverriddenOpinion(Attr);
+		}
+	}
+
 	return true;
 }
 
@@ -4314,12 +4472,14 @@ bool UnrealToUsd::CreateComponentPropertyBaker(
 		pxr::GfVec3f UsdMin{
 			FMath::Min(UEBoundsMinUsdSpace[0], UEBoundsMaxUsdSpace[0]),
 			FMath::Min(UEBoundsMinUsdSpace[1], UEBoundsMaxUsdSpace[1]),
-			FMath::Min(UEBoundsMinUsdSpace[2], UEBoundsMaxUsdSpace[2])};
+			FMath::Min(UEBoundsMinUsdSpace[2], UEBoundsMaxUsdSpace[2])
+		};
 
 		pxr::GfVec3f UsdMax{
 			FMath::Max(UEBoundsMinUsdSpace[0], UEBoundsMaxUsdSpace[0]),
 			FMath::Max(UEBoundsMinUsdSpace[1], UEBoundsMaxUsdSpace[1]),
-			FMath::Max(UEBoundsMinUsdSpace[2], UEBoundsMaxUsdSpace[2])};
+			FMath::Max(UEBoundsMinUsdSpace[2], UEBoundsMaxUsdSpace[2])
+		};
 
 		return pxr::VtArray<pxr::GfVec3f>{UsdMin, UsdMax};
 	};
@@ -4486,11 +4646,19 @@ bool UnrealToUsd::CreateComponentPropertyBaker(
 	if (const UCineCameraComponent* CameraComponent = Cast<UCineCameraComponent>(&Component))
 	{
 		static TSet<FString> RelevantProperties = {
-			TEXT("CurrentFocalLength"),
-			TEXT("FocusSettings.ManualFocusDistance"),
-			TEXT("CurrentAperture"),
-			TEXT("Filmback.SensorWidth"),
-			TEXT("Filmback.SensorHeight")};
+			UnrealIdentifiers::CurrentFocalLengthPropertyName.ToString(),
+			UnrealIdentifiers::ManualFocusDistancePropertyName.ToString(),
+			UnrealIdentifiers::CurrentAperturePropertyName.ToString(),
+			UnrealIdentifiers::SensorWidthPropertyName.ToString(),
+			UnrealIdentifiers::SensorHeightPropertyName.ToString(),
+			UnrealIdentifiers::SensorHorizontalOffsetPropertyName.ToString(),
+			UnrealIdentifiers::SensorVerticalOffsetPropertyName.ToString(),
+			UnrealIdentifiers::ExposureCompensationPropertyName.ToString(),
+			UnrealIdentifiers::ProjectionModePropertyName.ToString(),
+			UnrealIdentifiers::OrthoFarClipPlanePropertyName.ToString(),
+			UnrealIdentifiers::OrthoNearClipPlanePropertyName.ToString(),
+			UnrealIdentifiers::CustomNearClipppingPlanePropertyName.ToString()
+		};
 
 		if (RelevantProperties.Contains(PropertyPath))
 		{
@@ -4500,6 +4668,19 @@ bool UnrealToUsd::CreateComponentPropertyBaker(
 				return false;
 			}
 
+			TFunction<pxr::TfToken(ECameraProjectionMode::Type)> ProjectionConversion = [](ECameraProjectionMode::Type UEMode)
+			{
+				return UEMode == ECameraProjectionMode::Orthographic ? pxr::UsdGeomTokens->orthographic : pxr::UsdGeomTokens->perspective;
+			};
+
+			TFunction<pxr::GfVec2f(const FVector2D& UEVec)> ClippingRangeConversion = [StageInfo](const FVector2D& UEVec)
+			{
+				return pxr::GfVec2f{
+					UnrealToUsd::ConvertDistance(StageInfo, UEVec[0]),	  //
+					UnrealToUsd::ConvertDistance(StageInfo, UEVec[1])
+				};
+			};
+
 			TFunction<const UCineCameraComponent*(void)> ComponentGetter = CreateComponentGetter(CameraComponent);
 
 			pxr::UsdAttribute FocalLengthAttr = GeomCamera.CreateFocalLengthAttr();
@@ -4507,20 +4688,45 @@ bool UnrealToUsd::CreateComponentPropertyBaker(
 			pxr::UsdAttribute FStopAttr = GeomCamera.CreateFStopAttr();
 			pxr::UsdAttribute HorizontalApertureAttr = GeomCamera.CreateHorizontalApertureAttr();
 			pxr::UsdAttribute VerticalApertureAttr = GeomCamera.CreateVerticalApertureAttr();
+			pxr::UsdAttribute HorizontalApertureOffsetAttr = GeomCamera.CreateHorizontalApertureOffsetAttr();
+			pxr::UsdAttribute VerticalApertureOffsetAttr = GeomCamera.CreateVerticalApertureOffsetAttr();
+			pxr::UsdAttribute ExposureAttr = GeomCamera.CreateExposureAttr();
+			pxr::UsdAttribute ProjectionAttr = GeomCamera.CreateProjectionAttr();
+			pxr::UsdAttribute ClippingRangeAttr = GeomCamera.CreateClippingRangeAttr();
 
 			TFunction<void(float, double)> FocalLengthSetter = CreateCachedAttrSetter(FocalLengthAttr, DistanceConversion);
 			TFunction<void(float, double)> FocusDistanceSetter = CreateCachedAttrSetter(FocusDistanceAttr, DistanceConversion);
 			TFunction<void(float, double)> FStopSetter = CreateCachedAttrSetter(FStopAttr, NoFloatConversion);
 			TFunction<void(float, double)> HorizontalApertureSetter = CreateCachedAttrSetter(HorizontalApertureAttr, DistanceConversion);
 			TFunction<void(float, double)> VerticalApertureSetter = CreateCachedAttrSetter(VerticalApertureAttr, DistanceConversion);
+			TFunction<void(float, double)> HorizontalApertureOffsetSetter = CreateCachedAttrSetter(HorizontalApertureOffsetAttr, DistanceConversion);
+			TFunction<void(float, double)> VerticalApertureOffsetSetter = CreateCachedAttrSetter(VerticalApertureOffsetAttr, DistanceConversion);
+			TFunction<void(float, double)> ExposureSetter = CreateCachedAttrSetter(ExposureAttr, NoFloatConversion);
+			TFunction<void(ECameraProjectionMode::Type, double)> ProjectionSetter = CreateCachedAttrSetter(ProjectionAttr, ProjectionConversion);
+			TFunction<void(const FVector2D&, double)> ClippingRangeSetter = CreateCachedAttrSetter(ClippingRangeAttr, ClippingRangeConversion);
+
+			float DefaultUEFarPlane = FLT_MAX;
+			if (ClippingRangeAttr)
+			{
+				pxr::GfVec2f ClippingRangeUsdUnits{1, 1000000};	   // Values copied from the schema, as a fallback
+				ClippingRangeAttr.Get(&ClippingRangeUsdUnits, pxr::UsdTimeCode::EarliestTime());
+
+				DefaultUEFarPlane = UsdToUnreal::ConvertDistance(StageInfo, ClippingRangeUsdUnits[1]);
+			}
 
 			BakerType = EBakingType::Camera;
-			BakerFunction = [ComponentGetter,	 //
+			BakerFunction = [DefaultUEFarPlane,
+							 ComponentGetter,	 //
 							 FocalLengthSetter,
 							 FocusDistanceSetter,
 							 FStopSetter,
 							 HorizontalApertureSetter,
-							 VerticalApertureSetter](double UsdTimeCode) mutable
+							 VerticalApertureSetter,
+							 HorizontalApertureOffsetSetter,
+							 VerticalApertureOffsetSetter,
+							 ExposureSetter,
+							 ProjectionSetter,
+							 ClippingRangeSetter](double UsdTimeCode) mutable
 			{
 				TRACE_CPUPROFILER_EVENT_SCOPE(CameraBaker);
 
@@ -4535,6 +4741,36 @@ bool UnrealToUsd::CreateComponentPropertyBaker(
 				FStopSetter(CameraComponent->CurrentAperture, UsdTimeCode);
 				HorizontalApertureSetter(CameraComponent->Filmback.SensorWidth, UsdTimeCode);
 				VerticalApertureSetter(CameraComponent->Filmback.SensorHeight, UsdTimeCode);
+				HorizontalApertureOffsetSetter(CameraComponent->Filmback.SensorHorizontalOffset, UsdTimeCode);
+				VerticalApertureOffsetSetter(CameraComponent->Filmback.SensorVerticalOffset, UsdTimeCode);
+
+				if (CameraComponent->PostProcessSettings.bOverride_AutoExposureBias)
+				{
+					ExposureSetter(CameraComponent->PostProcessSettings.AutoExposureBias, UsdTimeCode);
+				}
+
+				ProjectionSetter(CameraComponent->ProjectionMode, UsdTimeCode);
+
+				bool bWriteClippingRange = false;
+				FVector2D UEClippingRange{0, DefaultUEFarPlane};
+				if (CameraComponent->ProjectionMode == ECameraProjectionMode::Perspective)
+				{
+					bWriteClippingRange = CameraComponent->bOverride_CustomNearClippingPlane;
+					if (bWriteClippingRange)
+					{
+						UEClippingRange[0] = CameraComponent->CustomNearClippingPlane;
+					}
+				}
+				else
+				{
+					bWriteClippingRange = true;
+					UEClippingRange[0] = CameraComponent->OrthoNearClipPlane;
+					UEClippingRange[1] = CameraComponent->OrthoFarClipPlane;
+				}
+				if (bWriteClippingRange)
+				{
+					ClippingRangeSetter(UEClippingRange, UsdTimeCode);
+				}
 			};
 		}
 	}
@@ -4831,7 +5067,8 @@ bool UnrealToUsd::CreateComponentPropertyBaker(
 		// attributes are uniform. The only animation we can export from this are the extents
 		static TSet<FString> RelevantProperties = {
 			GET_MEMBER_NAME_CHECKED(UUsdDrawModeComponent, BoundsMin).ToString(),
-			GET_MEMBER_NAME_CHECKED(UUsdDrawModeComponent, BoundsMax).ToString()};
+			GET_MEMBER_NAME_CHECKED(UUsdDrawModeComponent, BoundsMax).ToString()
+		};
 
 		if (RelevantProperties.Contains(PropertyPath))
 		{
@@ -5275,6 +5512,7 @@ UnrealToUsd::FPropertyTrackWriter UnrealToUsd::CreatePropertyTrackWriter(
 		if (pxr::UsdGeomCamera Camera{UsdPrim})
 		{
 			bool bConvertDistance = true;
+			TOptional<int32> ClippingPlaneIndex;
 
 			if (PropertyPath == UnrealIdentifiers::CurrentFocalLengthPropertyName)
 			{
@@ -5297,10 +5535,54 @@ UnrealToUsd::FPropertyTrackWriter UnrealToUsd::CreatePropertyTrackWriter(
 			{
 				Attr = Camera.CreateVerticalApertureAttr();
 			}
+			else if (PropertyPath == UnrealIdentifiers::SensorHorizontalOffsetPropertyName)
+			{
+				Attr = Camera.CreateHorizontalApertureOffsetAttr();
+			}
+			else if (PropertyPath == UnrealIdentifiers::SensorVerticalOffsetPropertyName)
+			{
+				Attr = Camera.CreateVerticalApertureOffsetAttr();
+			}
+			else if (PropertyPath == UnrealIdentifiers::ExposureCompensationPropertyName)
+			{
+				bConvertDistance = false;
+				Attr = Camera.CreateExposureAttr();
+			}
+			else if (PropertyPath == UnrealIdentifiers::OrthoFarClipPlanePropertyName)
+			{
+				Attr = Camera.CreateClippingRangeAttr();
+				ClippingPlaneIndex = 1;
+			}
+			else if (PropertyPath == UnrealIdentifiers::OrthoNearClipPlanePropertyName)
+			{
+				Attr = Camera.CreateClippingRangeAttr();
+				ClippingPlaneIndex = 0;
+			}
+			else if (PropertyPath == UnrealIdentifiers::CustomNearClipppingPlanePropertyName)
+			{
+				Attr = Camera.CreateClippingRangeAttr();
+				ClippingPlaneIndex = 0;
+			}
+
+			// TODO: Write out animation track for camera projection? Would need a new custom Enum writer type
 
 			if (Attr)
 			{
-				if (bConvertDistance)
+				if (ClippingPlaneIndex.IsSet() && ClippingPlaneIndex.GetValue() < 2)
+				{
+					Result.FloatWriter = [Attr, StageInfo, Index = ClippingPlaneIndex.GetValue()](float UEValue, double UsdTimeCode) mutable
+					{
+						// This is going to correspond to one track, from either near/far planes. So here we first read the
+						// current USD value to use as the base, and then just set either near/far
+						pxr::GfVec2f Value{1, 1000000};	   // Default values, copied directly from the schema as a fallback
+						Attr.Get<pxr::GfVec2f>(&Value);
+
+						Value[Index] = UnrealToUsd::ConvertDistance(StageInfo, UEValue);
+
+						Attr.Set(Value, UsdTimeCode);
+					};
+				}
+				else if (bConvertDistance)
 				{
 					Result.FloatWriter = [Attr, StageInfo](float UEValue, double UsdTimeCode) mutable
 					{
@@ -5760,11 +6042,13 @@ UnrealToUsd::FPropertyTrackWriter UnrealToUsd::CreatePropertyTrackWriter(
 					pxr::GfVec3f UsdMin{
 						FMath::Min(UEBoundsMinUsdSpace[0], UEBoundsMaxUsdSpace[0]),
 						FMath::Min(UEBoundsMinUsdSpace[1], UEBoundsMaxUsdSpace[1]),
-						FMath::Min(UEBoundsMinUsdSpace[2], UEBoundsMaxUsdSpace[2])};
+						FMath::Min(UEBoundsMinUsdSpace[2], UEBoundsMaxUsdSpace[2])
+					};
 					pxr::GfVec3f UsdMax{
 						FMath::Max(UEBoundsMinUsdSpace[0], UEBoundsMaxUsdSpace[0]),
 						FMath::Max(UEBoundsMinUsdSpace[1], UEBoundsMaxUsdSpace[1]),
-						FMath::Max(UEBoundsMinUsdSpace[2], UEBoundsMaxUsdSpace[2])};
+						FMath::Max(UEBoundsMinUsdSpace[2], UEBoundsMaxUsdSpace[2])
+					};
 					pxr::VtArray<pxr::GfVec3f> Extents{UsdMin, UsdMax};
 
 					Attr.Set(Extents, UsdTimeCode);
@@ -6083,6 +6367,28 @@ TArray<UE::FUsdAttribute> UsdUtils::GetAttributesForProperty(const UE::FUsdPrim&
 	{
 		return {UE::FUsdAttribute{UsdPrim.GetAttribute(pxr::UsdGeomTokens->verticalAperture)}};
 	}
+	else if (PropertyPath == SensorHorizontalOffsetPropertyName)
+	{
+		return {UE::FUsdAttribute{UsdPrim.GetAttribute(pxr::UsdGeomTokens->horizontalApertureOffset)}};
+	}
+	else if (PropertyPath == SensorVerticalOffsetPropertyName)
+	{
+		return {UE::FUsdAttribute{UsdPrim.GetAttribute(pxr::UsdGeomTokens->verticalApertureOffset)}};
+	}
+	else if (PropertyPath == ExposureCompensationPropertyName)
+	{
+		return {UE::FUsdAttribute{UsdPrim.GetAttribute(pxr::UsdGeomTokens->exposure)}};
+	}
+	else if (PropertyPath == ProjectionModePropertyName)
+	{
+		return {UE::FUsdAttribute{UsdPrim.GetAttribute(pxr::UsdGeomTokens->projection)}};
+	}
+	else if (PropertyPath == OrthoFarClipPlanePropertyName		  //
+			 || PropertyPath == OrthoNearClipPlanePropertyName	  //
+			 || PropertyPath == CustomNearClipppingPlanePropertyName)
+	{
+		return {UE::FUsdAttribute{UsdPrim.GetAttribute(pxr::UsdGeomTokens->clippingRange)}};
+	}
 
 	// Light attributes
 	else if (PropertyPath == IntensityPropertyName)
@@ -6093,20 +6399,23 @@ TArray<UE::FUsdAttribute> UsdUtils::GetAttributesForProperty(const UE::FUsdPrim&
 				UE::FUsdAttribute{UsdPrim.GetAttribute(pxr::UsdLuxTokens->inputsIntensity)},
 				UE::FUsdAttribute{UsdPrim.GetAttribute(pxr::UsdLuxTokens->inputsExposure)},
 				UE::FUsdAttribute{UsdPrim.GetAttribute(pxr::UsdLuxTokens->inputsWidth)},
-				UE::FUsdAttribute{UsdPrim.GetAttribute(pxr::UsdLuxTokens->inputsHeight)}};
+				UE::FUsdAttribute{UsdPrim.GetAttribute(pxr::UsdLuxTokens->inputsHeight)}
+			};
 		}
 		else if (UsdPrim.IsA<pxr::UsdLuxDiskLight>())
 		{
 			return {
 				UE::FUsdAttribute{UsdPrim.GetAttribute(pxr::UsdLuxTokens->inputsIntensity)},
 				UE::FUsdAttribute{UsdPrim.GetAttribute(pxr::UsdLuxTokens->inputsExposure)},
-				UE::FUsdAttribute{UsdPrim.GetAttribute(pxr::UsdLuxTokens->inputsRadius)}};
+				UE::FUsdAttribute{UsdPrim.GetAttribute(pxr::UsdLuxTokens->inputsRadius)}
+			};
 		}
 		else if (UsdPrim.IsA<pxr::UsdLuxDistantLight>())
 		{
 			return {
 				UE::FUsdAttribute{UsdPrim.GetAttribute(pxr::UsdLuxTokens->inputsIntensity)},
-				UE::FUsdAttribute{UsdPrim.GetAttribute(pxr::UsdLuxTokens->inputsExposure)}};
+				UE::FUsdAttribute{UsdPrim.GetAttribute(pxr::UsdLuxTokens->inputsExposure)}
+			};
 		}
 		else if (UsdPrim.IsA<pxr::UsdLuxSphereLight>())
 		{
@@ -6117,14 +6426,16 @@ TArray<UE::FUsdAttribute> UsdUtils::GetAttributesForProperty(const UE::FUsdPrim&
 					UE::FUsdAttribute{UsdPrim.GetAttribute(pxr::UsdLuxTokens->inputsExposure)},
 					UE::FUsdAttribute{UsdPrim.GetAttribute(pxr::UsdLuxTokens->inputsRadius)},
 					UE::FUsdAttribute{UsdPrim.GetAttribute(pxr::UsdLuxTokens->inputsShapingConeAngle)},
-					UE::FUsdAttribute{UsdPrim.GetAttribute(pxr::UsdLuxTokens->inputsShapingConeSoftness)}};
+					UE::FUsdAttribute{UsdPrim.GetAttribute(pxr::UsdLuxTokens->inputsShapingConeSoftness)}
+				};
 			}
 			else
 			{
 				return {
 					UE::FUsdAttribute{UsdPrim.GetAttribute(pxr::UsdLuxTokens->inputsIntensity)},
 					UE::FUsdAttribute{UsdPrim.GetAttribute(pxr::UsdLuxTokens->inputsExposure)},
-					UE::FUsdAttribute{UsdPrim.GetAttribute(pxr::UsdLuxTokens->inputsRadius)}};
+					UE::FUsdAttribute{UsdPrim.GetAttribute(pxr::UsdLuxTokens->inputsRadius)}
+				};
 			}
 		}
 	}
@@ -6174,7 +6485,8 @@ TArray<UE::FUsdAttribute> UsdUtils::GetAttributesForProperty(const UE::FUsdPrim&
 	{
 		return {
 			UE::FUsdAttribute{UsdPrim.GetAttribute(pxr::UsdLuxTokens->inputsShapingConeAngle)},
-			UE::FUsdAttribute{UsdPrim.GetAttribute(pxr::UsdLuxTokens->inputsShapingConeSoftness)}};
+			UE::FUsdAttribute{UsdPrim.GetAttribute(pxr::UsdLuxTokens->inputsShapingConeSoftness)}
+		};
 	}
 	else if (PropertyPath == LightSourceAnglePropertyName)
 	{
@@ -6184,7 +6496,8 @@ TArray<UE::FUsdAttribute> UsdUtils::GetAttributesForProperty(const UE::FUsdPrim&
 	// Bounds component properties
 	// For now we only support bounds animations, but there's nothing preventing us from supporting color/texture/draw mode animations here,
 	// especially since our component works with a "dynamic mesh" approach that already rebuilds its proxy on-demand anyway
-	else if (PropertyPath == GET_MEMBER_NAME_CHECKED(UUsdDrawModeComponent, BoundsMin) || PropertyPath == GET_MEMBER_NAME_CHECKED(UUsdDrawModeComponent, BoundsMax))
+	else if (PropertyPath == GET_MEMBER_NAME_CHECKED(UUsdDrawModeComponent, BoundsMin)
+			 || PropertyPath == GET_MEMBER_NAME_CHECKED(UUsdDrawModeComponent, BoundsMax))
 	{
 		// If we have a model API, let's direct the caller to `extentsHint` instead, as that has priority over
 		// `extent` in case the model API is applied to a boundable and both are authored
@@ -6257,6 +6570,11 @@ TArray<FName> UsdUtils::GetPropertiesForAttribute(const UE::FUsdPrim& Prim, cons
 		{pxr::UsdGeomTokens->fStop, 						{CurrentAperturePropertyName}},
 		{pxr::UsdGeomTokens->horizontalAperture, 			{SensorWidthPropertyName}},
 		{pxr::UsdGeomTokens->verticalAperture, 				{SensorHeightPropertyName}},
+		{pxr::UsdGeomTokens->horizontalApertureOffset, 		{SensorHorizontalOffsetPropertyName}},
+		{pxr::UsdGeomTokens->verticalApertureOffset, 		{SensorVerticalOffsetPropertyName}},
+		{pxr::UsdGeomTokens->exposure, 						{ExposureCompensationPropertyName}},
+		{pxr::UsdGeomTokens->projection, 					{ProjectionModePropertyName}},
+		{pxr::UsdGeomTokens->clippingRange, 				{OrthoFarClipPlanePropertyName, OrthoNearClipPlanePropertyName, CustomNearClipppingPlanePropertyName}},
 
 		{pxr::UsdLuxTokens->inputsIntensity, 				{IntensityPropertyName}},
 		{pxr::UsdLuxTokens->inputsExposure, 				{IntensityPropertyName}}, // In USD true intensity is just a function of "intensity" and "exposure"
@@ -6380,11 +6698,13 @@ bool UnrealToUsd::ConvertDrawModeComponent(
 		pxr::GfVec3f UsdMin{
 			FMath::Min(UEBoundsMinUsdSpace[0], UEBoundsMaxUsdSpace[0]),
 			FMath::Min(UEBoundsMinUsdSpace[1], UEBoundsMaxUsdSpace[1]),
-			FMath::Min(UEBoundsMinUsdSpace[2], UEBoundsMaxUsdSpace[2])};
+			FMath::Min(UEBoundsMinUsdSpace[2], UEBoundsMaxUsdSpace[2])
+		};
 		pxr::GfVec3f UsdMax{
 			FMath::Max(UEBoundsMinUsdSpace[0], UEBoundsMaxUsdSpace[0]),
 			FMath::Max(UEBoundsMinUsdSpace[1], UEBoundsMaxUsdSpace[1]),
-			FMath::Max(UEBoundsMinUsdSpace[2], UEBoundsMaxUsdSpace[2])};
+			FMath::Max(UEBoundsMinUsdSpace[2], UEBoundsMaxUsdSpace[2])
+		};
 
 		pxr::VtArray<pxr::GfVec3f> Extents{UsdMin, UsdMax};
 
@@ -6628,7 +6948,8 @@ namespace UE::USDPrimConversion::Private
 			// definition itself and we'll likely run into trouble if we try writing anything that differs from it
 			const static TSet<FString> FieldsToSkip = {
 				UsdToUnreal::ConvertToken(pxr::SdfFieldKeys->Specifier),
-				UsdToUnreal::ConvertToken(pxr::SdfFieldKeys->TypeName)};
+				UsdToUnreal::ConvertToken(pxr::SdfFieldKeys->TypeName)
+			};
 			if (FieldsToSkip.Contains(FullKeyPath))
 			{
 				continue;
diff --git a/Engine/Plugins/Runtime/USDCore/Source/UnrealUSDWrapper/Private/UnrealUSDWrapper.cpp b/Engine/Plugins/Runtime/USDCore/Source/UnrealUSDWrapper/Private/UnrealUSDWrapper.cpp
--- a/Engine/Plugins/Runtime/USDCore/Source/UnrealUSDWrapper/Private/UnrealUSDWrapper.cpp
+++ b/Engine/Plugins/Runtime/USDCore/Source/UnrealUSDWrapper/Private/UnrealUSDWrapper.cpp
@@ -441,6 +441,13 @@ FName UnrealIdentifiers::ManualFocusDistancePropertyName = GET_MEMBER_NAME_CHECK
 FName UnrealIdentifiers::CurrentAperturePropertyName = GET_MEMBER_NAME_CHECKED(UCineCameraComponent, CurrentAperture);
 FName UnrealIdentifiers::SensorWidthPropertyName = GET_MEMBER_NAME_CHECKED(UCineCameraComponent, Filmback.SensorWidth);
 FName UnrealIdentifiers::SensorHeightPropertyName = GET_MEMBER_NAME_CHECKED(UCineCameraComponent, Filmback.SensorHeight);
+FName UnrealIdentifiers::SensorHorizontalOffsetPropertyName = GET_MEMBER_NAME_CHECKED(UCineCameraComponent, Filmback.SensorHorizontalOffset);
+FName UnrealIdentifiers::SensorVerticalOffsetPropertyName = GET_MEMBER_NAME_CHECKED(UCineCameraComponent, Filmback.SensorVerticalOffset);
+FName UnrealIdentifiers::ExposureCompensationPropertyName = GET_MEMBER_NAME_CHECKED(UCineCameraComponent, PostProcessSettings.AutoExposureBias);
+FName UnrealIdentifiers::ProjectionModePropertyName = GET_MEMBER_NAME_CHECKED(UCineCameraComponent, ProjectionMode);
+FName UnrealIdentifiers::OrthoFarClipPlanePropertyName = GET_MEMBER_NAME_CHECKED(UCineCameraComponent, OrthoFarClipPlane);
+FName UnrealIdentifiers::OrthoNearClipPlanePropertyName = GET_MEMBER_NAME_CHECKED(UCineCameraComponent, OrthoNearClipPlane);
+FName UnrealIdentifiers::CustomNearClipppingPlanePropertyName = GET_MEMBER_NAME_CHECKED(UCineCameraComponent, CustomNearClippingPlane);
 
 FName UnrealIdentifiers::IntensityPropertyName = GET_MEMBER_NAME_CHECKED(ULightComponentBase, Intensity);
 FName UnrealIdentifiers::LightColorPropertyName = GET_MEMBER_NAME_CHECKED(ULightComponentBase, LightColor);
diff --git a/Engine/Plugins/Runtime/USDCore/Source/UnrealUSDWrapper/Public/UnrealUSDWrapper.h b/Engine/Plugins/Runtime/USDCore/Source/UnrealUSDWrapper/Public/UnrealUSDWrapper.h
--- a/Engine/Plugins/Runtime/USDCore/Source/UnrealUSDWrapper/Public/UnrealUSDWrapper.h
+++ b/Engine/Plugins/Runtime/USDCore/Source/UnrealUSDWrapper/Public/UnrealUSDWrapper.h
@@ -515,6 +515,13 @@ namespace UnrealIdentifiers
 	extern UNREALUSDWRAPPER_API FName CurrentAperturePropertyName;
 	extern UNREALUSDWRAPPER_API FName SensorWidthPropertyName;
 	extern UNREALUSDWRAPPER_API FName SensorHeightPropertyName;
+	extern UNREALUSDWRAPPER_API FName SensorHorizontalOffsetPropertyName;
+	extern UNREALUSDWRAPPER_API FName SensorVerticalOffsetPropertyName;
+	extern UNREALUSDWRAPPER_API FName ExposureCompensationPropertyName;
+	extern UNREALUSDWRAPPER_API FName ProjectionModePropertyName;
+	extern UNREALUSDWRAPPER_API FName OrthoFarClipPlanePropertyName;
+	extern UNREALUSDWRAPPER_API FName OrthoNearClipPlanePropertyName;
+	extern UNREALUSDWRAPPER_API FName CustomNearClipppingPlanePropertyName;
 
 	// LightComponentBase properties
 	extern UNREALUSDWRAPPER_API FName IntensityPropertyName;
