Fixed Ansi allocator for QAEnterprise project

diff --git a/Engine/Plugins/Runtime/USDCore/Source/UnrealUSDWrapper/UnrealUSDWrapper.Build.cs b/Engine/Plugins/Runtime/USDCore/Source/UnrealUSDWrapper/UnrealUSDWrapper.Build.cs
--- a/Engine/Plugins/Runtime/USDCore/Source/UnrealUSDWrapper/UnrealUSDWrapper.Build.cs
+++ b/Engine/Plugins/Runtime/USDCore/Source/UnrealUSDWrapper/UnrealUSDWrapper.Build.cs
@@ -213,7 +213,7 @@ namespace UnrealBuildTool.Rules
 			//			public YourProject(TargetInfo Target) : base(Target)
 			//			{
 			//				...
-			//				GlobalDefinitions.Add("FORCE_ANSI_ALLOCATOR=1");
+			//				StaticAllocator = StaticAllocatorType.Ansi;
 			//				GlobalDefinitions.Add("UE_USE_MALLOC_FILL_BYTES=0");
 			//				...
 			//			}
