USD: Fix logic error when detecting whether to enable the USD SDK or not;

Prevent some TSAN false positives due to TBB on Linux.

diff --git a/Engine/Plugins/Runtime/USDCore/Source/UnrealUSDWrapper/Private/UnrealUSDWrapper.cpp b/Engine/Plugins/Runtime/USDCore/Source/UnrealUSDWrapper/Private/UnrealUSDWrapper.cpp
--- a/Engine/Plugins/Runtime/USDCore/Source/UnrealUSDWrapper/Private/UnrealUSDWrapper.cpp
+++ b/Engine/Plugins/Runtime/USDCore/Source/UnrealUSDWrapper/Private/UnrealUSDWrapper.cpp
@@ -520,7 +520,7 @@ namespace UsdWrapperUtils
 class FUsdDiagnosticDelegate : public pxr::TfDiagnosticMgr::Delegate
 {
 public:
-	virtual ~FUsdDiagnosticDelegate() override{};
+	virtual ~FUsdDiagnosticDelegate() override {};
 
 	virtual void IssueError(const pxr::TfError& Error) override
 	{
@@ -602,12 +602,12 @@ namespace UE::UnrealUSDWrapper::Private
 	}
 };
 
-TArray<FString> UnrealUSDWrapper::RegisterPlugins(const FString& PathToPlugInfo)
+TArray<FString> UnrealUSDWrapper::RegisterPlugins(const FString& PathToPlugInfo) TSAN_SAFE
 {
 	return RegisterPlugins(TArray<FString>{PathToPlugInfo});
 }
 
-TArray<FString> UnrealUSDWrapper::RegisterPlugins(const TArray<FString>& PathsToPlugInfo)
+TArray<FString> UnrealUSDWrapper::RegisterPlugins(const TArray<FString>& PathsToPlugInfo) TSAN_SAFE
 {
 	UE::UnrealUSDWrapper::Private::EnsureModuleIsLoaded();
 
@@ -635,7 +635,7 @@ TArray<FString> UnrealUSDWrapper::RegisterPlugins(const TArray<FString>& PathsTo
 	return PluginNames;
 }
 
-TArray<FString> UnrealUSDWrapper::GetAllSupportedFileFormats()
+TArray<FString> UnrealUSDWrapper::GetAllSupportedFileFormats() TSAN_SAFE
 {
 	UE::UnrealUSDWrapper::Private::EnsureModuleIsLoaded();
 
@@ -661,7 +661,7 @@ TArray<FString> UnrealUSDWrapper::GetAllSupportedFileFormats()
 	return Result;
 }
 
-TArray<FString> UnrealUSDWrapper::GetNativeFileFormats()
+TArray<FString> UnrealUSDWrapper::GetNativeFileFormats() TSAN_SAFE
 {
 	UE::UnrealUSDWrapper::Private::EnsureModuleIsLoaded();
 
@@ -674,7 +674,8 @@ TArray<FString> UnrealUSDWrapper::GetNativeFileFormats()
 		pxr::UsdUsdFileFormatTokens->Id,
 		pxr::UsdUsdaFileFormatTokens->Id,
 		pxr::UsdUsdcFileFormatTokens->Id,
-		pxr::UsdUsdzFileFormatTokens->Id};
+		pxr::UsdUsdzFileFormatTokens->Id
+	};
 
 	std::set<std::string> FileExtensions;
 
@@ -701,7 +702,7 @@ TArray<FString> UnrealUSDWrapper::GetNativeFileFormats()
 	return Result;
 }
 
-void UnrealUSDWrapper::GetNativeFileFormats(TArray<FString>& OutTextFormats, TArray<FString>& OutPossiblyBinaryFormats)
+void UnrealUSDWrapper::GetNativeFileFormats(TArray<FString>& OutTextFormats, TArray<FString>& OutPossiblyBinaryFormats) TSAN_SAFE
 {
 	UE::UnrealUSDWrapper::Private::EnsureModuleIsLoaded();
 
@@ -727,7 +728,8 @@ void UnrealUSDWrapper::GetNativeFileFormats(TArray<FString>& OutTextFormats, TAr
 	const static std::vector<pxr::TfToken> BinaryFormatIds{
 		pxr::UsdUsdFileFormatTokens->Id,
 		pxr::UsdUsdcFileFormatTokens->Id,
-		pxr::UsdUsdzFileFormatTokens->Id};
+		pxr::UsdUsdzFileFormatTokens->Id
+	};
 
 	std::set<std::string> PossiblyBinaryExtensions;
 	for (const pxr::TfToken& FormatId : BinaryFormatIds)
@@ -1091,7 +1093,7 @@ void UnrealUSDWrapper::SetDefaultResolverDefaultSearchPath(const TArray<FDirecto
 #endif	  // #if USE_USD_SDK
 }
 
-void UnrealUSDWrapper::SetupDiagnosticDelegate()
+void UnrealUSDWrapper::SetupDiagnosticDelegate() TSAN_SAFE
 {
 #if USE_USD_SDK
 	if (Delegate.IsValid())
@@ -1106,7 +1108,7 @@ void UnrealUSDWrapper::SetupDiagnosticDelegate()
 #endif	  // USE_USD_SDK
 }
 
-void UnrealUSDWrapper::ClearDiagnosticDelegate()
+void UnrealUSDWrapper::ClearDiagnosticDelegate() TSAN_SAFE
 {
 #if USE_USD_SDK
 	if (!Delegate.IsValid())
@@ -1126,7 +1128,7 @@ class FUnrealUSDWrapperModule : public IUnrealUSDWrapperModule
 	FDelegateHandle OnProjectSettingsChangedHandle;
 
 public:
-	virtual void StartupModule() override
+	virtual void StartupModule() override TSAN_SAFE
 	{
 #if USE_USD_SDK
 		LLM_SCOPE_BYTAG(Usd);
diff --git a/Engine/Plugins/Runtime/USDCore/Source/UnrealUSDWrapper/UnrealUSDWrapper.Build.cs b/Engine/Plugins/Runtime/USDCore/Source/UnrealUSDWrapper/UnrealUSDWrapper.Build.cs
--- a/Engine/Plugins/Runtime/USDCore/Source/UnrealUSDWrapper/UnrealUSDWrapper.Build.cs
+++ b/Engine/Plugins/Runtime/USDCore/Source/UnrealUSDWrapper/UnrealUSDWrapper.Build.cs
@@ -225,7 +225,10 @@ namespace UnrealBuildTool.Rules
 			// Also note that we're forced to use dynamic linking for monolithic targets mainly because static linking the USD libraries disables support for user USD plugins, and secondly
 			// because those static libraries would need to be linked with the --whole-archive argument, and there is currently no standard way of doing that in UE.
 			bool bIsGameOrMonolithic = Target.LinkType == TargetLinkType.Monolithic || Target.Type == TargetType.Client || Target.Type == TargetType.Server || Target.Type == TargetType.Game;
-			if (bEnableUsdSdk && bIsGameOrMonolithic && (Target.StaticAllocator != StaticAllocatorType.Ansi) && !Target.GlobalDefinitions.Contains("UE_USE_MALLOC_FILL_BYTES=0"))
+			if (bEnableUsdSdk &&
+				bIsGameOrMonolithic &&
+				((Target.StaticAllocator != StaticAllocatorType.Ansi) || !Target.GlobalDefinitions.Contains("UE_USE_MALLOC_FILL_BYTES=0"))
+			)
 			{
 				PublicDefinitions.Add("USD_FORCE_DISABLED=1");
 				bEnableUsdSdk = false;
