USD: Fix warnings about asset names differing from package names when opening USDZ stages with textures.

diff --git a/Engine/Plugins/Runtime/USDCore/Source/USDUtilities/Private/USDShadeConversion.cpp b/Engine/Plugins/Runtime/USDCore/Source/USDUtilities/Private/USDShadeConversion.cpp
--- a/Engine/Plugins/Runtime/USDCore/Source/USDUtilities/Private/USDShadeConversion.cpp
+++ b/Engine/Plugins/Runtime/USDCore/Source/USDUtilities/Private/USDShadeConversion.cpp
@@ -2,7 +2,6 @@
 
 #include "USDShadeConversion.h"
 
-
 #if USE_USD_SDK
 
 #include "UnrealUSDWrapper.h"
@@ -144,8 +143,8 @@ namespace UE::UsdShadeConversion::Private
 		FScopedUsdAllocs UsdAllocs;
 
 		FString PathToResolve = LayerHandle ? UsdToUnreal::ConvertString(
-									pxr::SdfComputeAssetPathRelativeToLayer(LayerHandle, UnrealToUsd::ConvertString(*AssetPath).Get())
-								)
+												  pxr::SdfComputeAssetPathRelativeToLayer(LayerHandle, UnrealToUsd::ConvertString(*AssetPath).Get())
+											  )
 											: AssetPath;
 
 		// We need to provide absolute paths to the resolver later: It has no idea what to do with a relative
@@ -251,6 +250,7 @@ namespace UE::UsdShadeConversion::Private
 		const FString& TextureExtension,
 		UTextureFactory* TextureFactory,
 		UObject* Outer,
+		FName TextureName,
 		EObjectFlags ObjectFlags
 	)
 	{
@@ -266,12 +266,6 @@ namespace UE::UsdShadeConversion::Private
 
 		FScopedUnrealAllocs UEAllocs;
 
-		FName TextureName = MakeUniqueObjectName(
-			Outer,
-			UTexture::StaticClass(),
-			*UsdUnreal::ObjectUtils::SanitizeObjectName(FPaths::GetBaseFilename(ResolvedTexturePath))
-		);
-
 		UObject* Context = nullptr;
 		FFeedbackContext* FeedbackContext = GWarn;
 		return Cast<UTexture>(TextureFactory->FactoryCreateBinary(
@@ -1128,7 +1122,8 @@ namespace UE::UsdShadeConversion::Private
 				TextureValue.UVScale.GetVector()[0],
 				TextureValue.UVScale.GetVector()[1],
 				TextureValue.UVTranslation[0],
-				TextureValue.UVTranslation[1]};
+				TextureValue.UVTranslation[1]
+			};
 			UsdUtils::SetVectorParameterValue(Material, *FString::Printf(TEXT("%sScaleTranslation"), ParameterName), ScaleAndTranslation);
 
 			UsdUtils::SetScalarParameterValue(Material, *FString::Printf(TEXT("%sRotation"), ParameterName), TextureValue.UVRotation);
@@ -1257,7 +1252,7 @@ namespace UE::UsdShadeConversion::Private
 			if (IsInsideUsdzArchive(ResolvedTexturePath, TextureExtension))
 			{
 				// Always prefer using the TextureFactory if we can, as it may provide compression, which the runtime version never will
-				Texture = ReadTextureFromUsdzArchiveEditor(ResolvedTexturePath, TextureExtension, TextureFactory, Outer, ObjectFlags);
+				Texture = ReadTextureFromUsdzArchiveEditor(ResolvedTexturePath, TextureExtension, TextureFactory, Outer, SanitizedName, ObjectFlags);
 			}
 			// Not inside an USDZ archive, just a regular texture
 			else
@@ -2148,6 +2143,7 @@ bool UsdToUnreal::FTextureParameterValue::GetSRGBValue() const
 #if WITH_EDITOR && PLATFORM_WINDOWS
 namespace UE::UsdShadeConversion::Private
 {
+	// clang-format off
 	std::unordered_set<std::string> MaterialXTypes
 	{
 		"boolean",
@@ -2189,6 +2185,7 @@ namespace UE::UsdShadeConversion::Private
 		{"vector4f", "vector4"},
 		{"int", "integer"},
 	};
+	// clang-format on
 
 	bool ComputeShaderGraph(mx::NodePtr Node, const pxr::UsdShadeShader& Shader)
 	{
@@ -2207,11 +2204,11 @@ namespace UE::UsdShadeConversion::Private
 
 			if (MaterialXTypes.find(InputType) == MaterialXTypes.end())
 			{
-				if(auto It = UsdToMaterialXTypes.find(InputType); It != UsdToMaterialXTypes.end())
+				if (auto It = UsdToMaterialXTypes.find(InputType); It != UsdToMaterialXTypes.end())
 				{
 					InputType = It->second;
 				}
-				else if (InputName == "file") // we're dealing with a texture, USD calls the type "asset" for the input
+				else if (InputName == "file")	 // we're dealing with a texture, USD calls the type "asset" for the input
 				{
 					InputType = "filename";
 				}
@@ -2228,7 +2225,7 @@ namespace UE::UsdShadeConversion::Private
 			for (const pxr::UsdShadeConnectionSourceInfo& Source : *ConnectedSources)
 			{
 				bHasValue = false;
-				pxr::UsdShadeShader ConnectedShader{ Source.source.GetPrim() };
+				pxr::UsdShadeShader ConnectedShader{Source.source.GetPrim()};
 
 				pxr::TfToken NodeDefTokenId;
 				ConnectedShader.GetShaderId(&NodeDefTokenId);
@@ -2256,10 +2253,10 @@ namespace UE::UsdShadeConversion::Private
 
 				// add the output of the connected node to the current input
 				TUsdStore<std::vector<pxr::UsdShadeOutput>> UsdOutputs = ConnectedShader.GetOutputs();
-				for(const pxr::UsdShadeOutput& UsdOutput : *UsdOutputs)
+				for (const pxr::UsdShadeOutput& UsdOutput : *UsdOutputs)
 				{
-					//no need to add the default output
-					if(std::string OutputName = UsdOutput.GetBaseName().GetString(); OutputName != "out")
+					// no need to add the default output
+					if (std::string OutputName = UsdOutput.GetBaseName().GetString(); OutputName != "out")
 					{
 						Input->setOutputString(OutputName);
 					}
@@ -2274,40 +2271,43 @@ namespace UE::UsdShadeConversion::Private
 			if (bHasValue)
 			{
 				auto SetInputValue = [&UsdInput, Node, Input](auto Type)
-					{
-						using MtlxType = decltype(Type);
-
-						// complex template to just retrieve the right type, basically a mapping between Usd and Mtlx for Vector types, otherwise we just take the type (float, int, string)
-						using UsdType = std::conditional_t<
-							std::is_same_v<MtlxType, mx::Color4> || std::is_same_v<MtlxType, mx::Vector4>,
-							pxr::GfVec4f,
+				{
+					using MtlxType = decltype(Type);
+
+					// complex template to just retrieve the right type, basically a mapping between Usd and Mtlx for Vector types, otherwise we just take the type (float, int, string)
+					// clang-format off
+					using UsdType = std::conditional_t<
+						std::is_same_v<MtlxType, mx::Color4> || std::is_same_v<MtlxType, mx::Vector4>,
+						pxr::GfVec4f,
+							std::conditional_t<
+							std::is_same_v<MtlxType, mx::Color3> || std::is_same_v<MtlxType, mx::Vector3>,
+								pxr::GfVec3f,
 								std::conditional_t<
-								std::is_same_v<MtlxType, mx::Color3> || std::is_same_v<MtlxType, mx::Vector3>,
-									pxr::GfVec3f,
-									std::conditional_t<
-										std::is_same_v<MtlxType, mx::Vector2>,
-										pxr::GfVec2f,
-							MtlxType>>>;
-
-						if (UsdType Value; UsdInput.Get<UsdType>(&Value))
+									std::is_same_v<MtlxType, mx::Vector2>,
+									pxr::GfVec2f,
+						MtlxType>>>;
+					// clang-format on
+
+					if (UsdType Value; UsdInput.Get<UsdType>(&Value))
+					{
+						if constexpr (std::is_same_v<UsdType, pxr::GfVec4f> || std::is_same_v<UsdType, pxr::GfVec3f>
+									  || std::is_same_v<UsdType, pxr::GfVec2f>)
 						{
-							if constexpr(std::is_same_v<UsdType, pxr::GfVec4f> || std::is_same_v<UsdType, pxr::GfVec3f> || std::is_same_v<UsdType, pxr::GfVec2f>)
-							{
-								Input->setValue(MtlxType{ Value.data(), Value.data() + UsdType::dimension });
-							}
-							else
-							{
-								Input->setValue(Value);
-							}
+							Input->setValue(MtlxType{Value.data(), Value.data() + UsdType::dimension});
 						}
 						else
 						{
-							mx::NodeDefPtr NodeDef = Node->getNodeDef();
-							mx::InputPtr ActiveInput = NodeDef->getActiveInput(Input->getName());
-							MtlxType DefaultValue = ActiveInput->getDefaultValue()->asA<MtlxType>();
-							Input->setValue(DefaultValue);
+							Input->setValue(Value);
 						}
-					};
+					}
+					else
+					{
+						mx::NodeDefPtr NodeDef = Node->getNodeDef();
+						mx::InputPtr ActiveInput = NodeDef->getActiveInput(Input->getName());
+						MtlxType DefaultValue = ActiveInput->getDefaultValue()->asA<MtlxType>();
+						Input->setValue(DefaultValue);
+					}
+				};
 
 				if (InputType == "float")
 				{
@@ -2357,11 +2357,14 @@ namespace UE::UsdShadeConversion::Private
 					UsdInput.Get<pxr::SdfAssetPath>(&AssetPath);
 					TUsdStore<UE::FSdfLayer> Layer = UsdUtils::FindLayerForAttribute(UsdInput.GetAttr(), pxr::UsdTimeCode::Default().GetValue());
 
-					FString PathToResolve =
-						Layer.Get()
-						? UsdToUnreal::ConvertString(TUsdStore<std::string>(pxr::SdfComputeAssetPathRelativeToLayer(Layer.Get(), AssetPath.GetAssetPath())).Get())
-						: AssetPath.GetAssetPath().c_str();
-					Input->setValue(std::string{ TCHAR_TO_UTF8(*PathToResolve) });
+					FString PathToResolve = Layer.Get() ? UsdToUnreal::ConvertString(
+															  TUsdStore<std::string>(
+																  pxr::SdfComputeAssetPathRelativeToLayer(Layer.Get(), AssetPath.GetAssetPath())
+															  )
+																  .Get()
+														  )
+														: AssetPath.GetAssetPath().c_str();
+					Input->setValue(std::string{TCHAR_TO_UTF8(*PathToResolve)});
 				}
 				else
 				{
@@ -2372,11 +2375,11 @@ namespace UE::UsdShadeConversion::Private
 
 		return true;
 	}
-}
+}	 // namespace UE::UsdShadeConversion::Private
 
 mx::DocumentPtr UsdToUnreal::CreateSurfaceMaterial(const pxr::UsdPrim& UsdShadeMaterialPrim, const TCHAR* RenderContext)
 {
-	pxr::UsdShadeMaterial UsdShadeMaterial{ UsdShadeMaterialPrim };
+	pxr::UsdShadeMaterial UsdShadeMaterial{UsdShadeMaterialPrim};
 
 	if (!UsdShadeMaterial)
 	{
@@ -2401,12 +2404,10 @@ mx::DocumentPtr UsdToUnreal::CreateSurfaceMaterial(const pxr::UsdPrim& UsdShadeM
 	mx::DocumentPtr Document = mx::createDocument();
 	{
 		mx::DocumentPtr MaterialXLibrary = mx::createDocument();
-		mx::FileSearchPath MaterialXFolder{ TCHAR_TO_UTF8(*FPaths::Combine(
-			FPaths::EngineDir(),
-			TEXT("Binaries"),
-			TEXT("ThirdParty"),
-			TEXT("MaterialX"))) };
-		mx::StringSet LoadedLibs = mx::loadLibraries({ "libraries" }, MaterialXFolder, MaterialXLibrary);
+		mx::FileSearchPath MaterialXFolder{
+			TCHAR_TO_UTF8(*FPaths::Combine(FPaths::EngineDir(), TEXT("Binaries"), TEXT("ThirdParty"), TEXT("MaterialX")))
+		};
+		mx::StringSet LoadedLibs = mx::loadLibraries({"libraries"}, MaterialXFolder, MaterialXLibrary);
 		Document->importLibrary(MaterialXLibrary);
 	}
 
@@ -2424,7 +2425,7 @@ mx::DocumentPtr UsdToUnreal::CreateSurfaceMaterial(const pxr::UsdPrim& UsdShadeM
 
 		mx::NodePtr MtlxSurfaceShader = Document->addNodeInstance(NodeDef, SurfaceShader.GetPrim().GetName().GetString());
 
-		//if the shadergraph is not valid (probably caused by a missing nodedef or wrong type), just invalide it all
+		// if the shadergraph is not valid (probably caused by a missing nodedef or wrong type), just invalide it all
 		if (!UE::UsdShadeConversion::Private::ComputeShaderGraph(MtlxSurfaceShader, SurfaceShader))
 		{
 			MtlxSurfaceShader = nullptr;
@@ -2433,7 +2434,7 @@ mx::DocumentPtr UsdToUnreal::CreateSurfaceMaterial(const pxr::UsdPrim& UsdShadeM
 		return MtlxSurfaceShader;
 	};
 
-	//Create the shader graph for the surface shader and displacement shader
+	// Create the shader graph for the surface shader and displacement shader
 	if (mx::NodePtr MaterialXSurfaceShader = CreateShaderGraph(SurfaceShader, Document))
 	{
 		mx::NodeDefPtr SurfaceMaterialNodeDef = Document->getNodeDef("ND_surfacematerial");
@@ -2441,10 +2442,10 @@ mx::DocumentPtr UsdToUnreal::CreateSurfaceMaterial(const pxr::UsdPrim& UsdShadeM
 
 		SurfaceMaterial->setConnectedNode("surfaceshader", MaterialXSurfaceShader);
 
-		//the displacement shader is not necessarily required
+		// the displacement shader is not necessarily required
 		if (DisplacementShader)
 		{
-			if(mx::NodePtr MaterialXDisplacementShader = CreateShaderGraph(DisplacementShader, Document); MaterialXDisplacementShader)
+			if (mx::NodePtr MaterialXDisplacementShader = CreateShaderGraph(DisplacementShader, Document); MaterialXDisplacementShader)
 			{
 				SurfaceMaterial->setConnectedNode("displacementshader", MaterialXDisplacementShader);
 			}
@@ -2461,7 +2462,7 @@ mx::DocumentPtr UsdToUnreal::CreateSurfaceMaterial(const pxr::UsdPrim& UsdShadeM
 
 	return Document;
 }
-#endif // WITH_EDITOR && PLATFORM_WINDOWS
+#endif	  // WITH_EDITOR && PLATFORM_WINDOWS
 
 bool UsdToUnreal::ConvertMaterial(
 	const pxr::UsdPrim& UsdShadeMaterialPrim,
@@ -2831,9 +2832,8 @@ bool UsdToUnreal::ConvertShadeInputsToParameters(
 				bHasMaterialInfo = true;
 			}
 		}
-		else if ( ShadeInput.GetTypeName() == pxr::SdfValueTypeNames->Float ||
-			ShadeInput.GetTypeName() == pxr::SdfValueTypeNames->Double ||
-			ShadeInput.GetTypeName() == pxr::SdfValueTypeNames->Half )
+		else if (ShadeInput.GetTypeName() == pxr::SdfValueTypeNames->Float || ShadeInput.GetTypeName() == pxr::SdfValueTypeNames->Double
+				 || ShadeInput.GetTypeName() == pxr::SdfValueTypeNames->Half)
 		{
 			if (GetFloatParameterValue(Connectable, ShadeInput.GetBaseName(), 1.f, ParameterValue))
 			{
@@ -3365,9 +3365,10 @@ TArray<FString> UsdUtils::GetMaterialXFilePaths(const pxr::UsdPrim& Prim)
 
 	TArray<FString> Result;
 
-	auto AddMaterialXFilePath = [&Result](const pxr::UsdPrimCompositionQueryArc& CompositionArc, auto ReferenceOrPayloadEditor, auto UsdReferenceOrPayload)
+	auto AddMaterialXFilePath =
+		[&Result](const pxr::UsdPrimCompositionQueryArc& CompositionArc, auto ReferenceOrPayloadEditor, auto UsdReferenceOrPayload)
 	{
-		if(CompositionArc.GetIntroducingListEditor(&ReferenceOrPayloadEditor, &UsdReferenceOrPayload))
+		if (CompositionArc.GetIntroducingListEditor(&ReferenceOrPayloadEditor, &UsdReferenceOrPayload))
 		{
 			// The mtlx file will be the "target layer". It's useful to get its real path via the SdfLayer
 			// interface as it already makes sure it is absolute
@@ -3375,7 +3376,7 @@ TArray<FString> UsdUtils::GetMaterialXFilePaths(const pxr::UsdPrim& Prim)
 			FString AbsoluteFilePath = UsdToUnreal::ConvertString(TargetLayer->GetRealPath());
 
 			FString Extension = FPaths::GetExtension(AbsoluteFilePath);
-			if(Extension == UnrealIdentifiers::MaterialXRenderContext && FPaths::FileExists(AbsoluteFilePath))
+			if (Extension == UnrealIdentifiers::MaterialXRenderContext && FPaths::FileExists(AbsoluteFilePath))
 			{
 				Result.Add(AbsoluteFilePath);
 			}
@@ -3384,14 +3385,14 @@ TArray<FString> UsdUtils::GetMaterialXFilePaths(const pxr::UsdPrim& Prim)
 
 	// We used to just fetch "direct references" here, but stages may compose the .mtlx file reference onto the prim
 	// via another sublayer or reference, so it will be marked as an "ancestral arc" and not be included in the "direct references" filter
-	pxr::UsdPrimCompositionQuery PrimCompositionQuery = pxr::UsdPrimCompositionQuery{ Prim };
-	for(const pxr::UsdPrimCompositionQueryArc& CompositionArc : PrimCompositionQuery.GetCompositionArcs())
+	pxr::UsdPrimCompositionQuery PrimCompositionQuery = pxr::UsdPrimCompositionQuery{Prim};
+	for (const pxr::UsdPrimCompositionQueryArc& CompositionArc : PrimCompositionQuery.GetCompositionArcs())
 	{
-		if(CompositionArc.GetArcType() == pxr::PcpArcTypeReference)
+		if (CompositionArc.GetArcType() == pxr::PcpArcTypeReference)
 		{
 			AddMaterialXFilePath(CompositionArc, pxr::SdfReferenceEditorProxy{}, pxr::SdfReference{});
 		}
-		else if(CompositionArc.GetArcType() == pxr::PcpArcType::PcpArcTypePayload)
+		else if (CompositionArc.GetArcType() == pxr::PcpArcType::PcpArcTypePayload)
 		{
 			AddMaterialXFilePath(CompositionArc, pxr::SdfPayloadEditorProxy{}, pxr::SdfPayload{});
 		}
@@ -3399,7 +3400,7 @@ TArray<FString> UsdUtils::GetMaterialXFilePaths(const pxr::UsdPrim& Prim)
 
 	return Result;
 }
-#endif // WITH_EDITOR
+#endif	  // WITH_EDITOR
 
 void UsdUtils::AuthorUnrealMaterialBinding(pxr::UsdPrim& MeshOrGeomSubsetPrim, const FString& UnrealMaterialPathName)
 {
