[USD] Applying AdditionalTransform to SkeletalMesh Vertices.

#rnx

diff --git a/Engine/Plugins/Runtime/USDCore/Source/USDUtilities/Private/USDSkeletalDataConversion.cpp b/Engine/Plugins/Runtime/USDCore/Source/USDUtilities/Private/USDSkeletalDataConversion.cpp
--- a/Engine/Plugins/Runtime/USDCore/Source/USDUtilities/Private/USDSkeletalDataConversion.cpp
+++ b/Engine/Plugins/Runtime/USDCore/Source/USDUtilities/Private/USDSkeletalDataConversion.cpp
@@ -1417,7 +1417,7 @@ bool UsdToUnreal::ConvertSkinnedMesh(
 		for (uint32 PointIndex = 0; PointIndex < NumPoints; ++PointIndex)
 		{
 			const GfVec3f& Point = UsdPoints[PointIndex];
-			SkelMeshImportData.Points[PointIndex + NumExistingPoints] = (FVector3f)UsdToUnreal::ConvertVector(StageInfo, Point);
+			SkelMeshImportData.Points[PointIndex + NumExistingPoints] = (FVector3f)CommonOptions.AdditionalTransform.TransformPosition(UsdToUnreal::ConvertVector(StageInfo, Point));
 		}
 	}
 
