USD:
- Add operator -> overload to TUsdStore

diff --git a/Engine/Plugins/Runtime/USDCore/Source/UnrealUSDWrapper/Public/USDMemory.h b/Engine/Plugins/Runtime/USDCore/Source/UnrealUSDWrapper/Public/USDMemory.h
--- a/Engine/Plugins/Runtime/USDCore/Source/UnrealUSDWrapper/Public/USDMemory.h
+++ b/Engine/Plugins/Runtime/USDCore/Source/UnrealUSDWrapper/Public/USDMemory.h
@@ -256,6 +256,16 @@ public:
 		return Get();
 	}
 
+	UsdObjectType* operator->()
+	{
+		return &Get();
+	}
+
+	const UsdObjectType* operator->() const
+	{
+		return &Get();
+	}
+
 private:
 	TOptional<UsdObjectType> StoredUsdObject;
 };
