USD: Fix divide by zero crash in case a polygon has less than 0 vertices.

diff --git a/Engine/Plugins/Runtime/USDCore/Source/USDUtilities/Private/USDSkeletalDataConversion.cpp b/Engine/Plugins/Runtime/USDCore/Source/USDUtilities/Private/USDSkeletalDataConversion.cpp
--- a/Engine/Plugins/Runtime/USDCore/Source/USDUtilities/Private/USDSkeletalDataConversion.cpp
+++ b/Engine/Plugins/Runtime/USDCore/Source/USDUtilities/Private/USDSkeletalDataConversion.cpp
@@ -1779,12 +1779,21 @@ bool UsdToUnreal::ConvertSkinnedMesh(
 
 	SkelMeshImportData.Wedges.Reserve((NumExistingFaces + NumFaces) * 6);
 
+	uint32 NumSkippedPolygons = 0;
+
 	uint32 NumProcessedFaceVertexIndices = 0;
 	for (uint32 PolygonIndex = NumExistingFaces, LocalIndex = 0; PolygonIndex < NumExistingFaces + NumFaces; ++PolygonIndex, ++LocalIndex)
 	{
 		const uint32 NumOriginalFaceVertices = FaceCounts[LocalIndex];
 		const uint32 NumFinalFaceVertices = 3;
 
+		// Skip "polygon" if it has less than 3 vertices
+		if (NumOriginalFaceVertices < 3)
+		{
+			++NumSkippedPolygons;
+			continue;
+		}
+
 		// Manage materials
 		int32 LocalMaterialIndex = 0;
 		if (FaceMaterialIndices.IsValidIndex(LocalIndex))
@@ -1923,6 +1932,17 @@ bool UsdToUnreal::ConvertSkinnedMesh(
 		NumProcessedFaceVertexIndices += NumOriginalFaceVertices;
 	}
 
+	if (NumSkippedPolygons > 0)
+	{
+		UE_LOG(
+			LogUsd,
+			Warning,
+			TEXT("Ignoring %d polygons with less than 3 vertices from mesh '%s'"),
+			NumSkippedPolygons,
+			*UsdToUnreal::ConvertPath(UsdMesh.GetPrim().GetPrimPath())
+		);
+	}
+
 	// Convert joint influences into the SkeletalMeshImportData
 
 	// ComputeJointInfluences returns the influences per bone that applies to all the points of the mesh
