USD: Use shorter file paths when exporting levels and assets.

diff --git a/Engine/Plugins/Importers/USDImporter/Source/USDExporter/Private/AnimSequenceExporterUSD.cpp b/Engine/Plugins/Importers/USDImporter/Source/USDExporter/Private/AnimSequenceExporterUSD.cpp
--- a/Engine/Plugins/Importers/USDImporter/Source/USDExporter/Private/AnimSequenceExporterUSD.cpp
+++ b/Engine/Plugins/Importers/USDImporter/Source/USDExporter/Private/AnimSequenceExporterUSD.cpp
@@ -6,6 +6,7 @@
 #include "UnrealUSDWrapper.h"
 #include "USDErrorUtils.h"
 #include "USDExporterModule.h"
+#include "USDExportUtils.h"
 #include "USDLayerUtils.h"
 #include "USDLog.h"
 #include "USDObjectUtils.h"
@@ -126,6 +127,8 @@ bool UAnimSequenceExporterUSD::ExportBinary(
 		return false;
 	}
 
+	UsdUnreal::ExportUtils::FUniquePathScope UniquePathScope;
+
 	// See comment on the analogous line within StaticMeshExporterUSD.cpp
 	ExportTask->bPrompt = false;
 
diff --git a/Engine/Plugins/Importers/USDImporter/Source/USDExporter/Private/GeometryCacheExporterUSD.cpp b/Engine/Plugins/Importers/USDImporter/Source/USDExporter/Private/GeometryCacheExporterUSD.cpp
--- a/Engine/Plugins/Importers/USDImporter/Source/USDExporter/Private/GeometryCacheExporterUSD.cpp
+++ b/Engine/Plugins/Importers/USDImporter/Source/USDExporter/Private/GeometryCacheExporterUSD.cpp
@@ -7,6 +7,7 @@
 #include "UnrealUSDWrapper.h"
 #include "USDConversionUtils.h"
 #include "USDExporterModule.h"
+#include "USDExportUtils.h"
 #include "USDGeomMeshConversion.h"
 #include "USDLog.h"
 #include "USDObjectUtils.h"
@@ -148,6 +149,8 @@ bool UGeometryCacheExporterUSD::ExportBinary(
 		return false;
 	}
 
+	UsdUnreal::ExportUtils::FUniquePathScope UniquePathScope;
+
 	// Get a simple GUID hash/identifier of our mesh
 	FString DDCKeyHash;
 	const FString& DDCKey = GeometryCache->GetHash();
diff --git a/Engine/Plugins/Importers/USDImporter/Source/USDExporter/Private/LevelSequenceExporterUSD.cpp b/Engine/Plugins/Importers/USDImporter/Source/USDExporter/Private/LevelSequenceExporterUSD.cpp
--- a/Engine/Plugins/Importers/USDImporter/Source/USDExporter/Private/LevelSequenceExporterUSD.cpp
+++ b/Engine/Plugins/Importers/USDImporter/Source/USDExporter/Private/LevelSequenceExporterUSD.cpp
@@ -9,6 +9,7 @@
 #include "USDConversionUtils.h"
 #include "USDErrorUtils.h"
 #include "USDExporterModule.h"
+#include "USDExportUtils.h"
 #include "USDGeomMeshConversion.h"
 #include "USDLayerUtils.h"
 #include "USDLog.h"
@@ -1107,10 +1108,10 @@ namespace UE::LevelSequenceExporterUSD::Private
 			{
 				FString PrimPath = PrimPathStr ? *PrimPathStr
 											   : UsdUtils::GetPrimPathForObject(
-												   &Component,
-												   TEXT(""),
-												   Context.ExportOptions && Context.ExportOptions->LevelExportOptions.bExportActorFolders
-											   );
+													 &Component,
+													 TEXT(""),
+													 Context.ExportOptions && Context.ExportOptions->LevelExportOptions.bExportActorFolders
+												 );
 				if (PrimPath.IsEmpty())
 				{
 					return {};
@@ -1836,6 +1837,7 @@ bool ULevelSequenceExporterUsd::ExportBinary(
 		return false;
 	}
 
+	UsdUnreal::ExportUtils::FUniquePathScope UniquePathScope;
 	FScopedUsdMessageLog UsdMessageLog;
 
 	FSequencerInitParams Params;
diff --git a/Engine/Plugins/Importers/USDImporter/Source/USDExporter/Private/MaterialExporterUSD.cpp b/Engine/Plugins/Importers/USDImporter/Source/USDExporter/Private/MaterialExporterUSD.cpp
--- a/Engine/Plugins/Importers/USDImporter/Source/USDExporter/Private/MaterialExporterUSD.cpp
+++ b/Engine/Plugins/Importers/USDImporter/Source/USDExporter/Private/MaterialExporterUSD.cpp
@@ -7,6 +7,7 @@
 #include "USDClassesModule.h"
 #include "USDConversionUtils.h"
 #include "USDExporterModule.h"
+#include "USDExportUtils.h"
 #include "USDGeomMeshConversion.h"
 #include "USDLog.h"
 #include "USDMetadataExportOptions.h"
@@ -281,6 +282,8 @@ bool UMaterialExporterUsd::ExportMaterial(
 		return false;
 	}
 
+	UsdUnreal::ExportUtils::FUniquePathScope UniquePathScope;
+
 	// Check if we already have exported what we plan on exporting anyway
 	if (FPaths::FileExists(FilePath.FilePath))
 	{
@@ -441,6 +444,8 @@ bool UMaterialExporterUsd::ExportMaterialsForStage(
 		return true;
 	}
 
+	UsdUnreal::ExportUtils::FUniquePathScope UniquePathScope;
+
 	const FString ExtensionNoDot = FPaths::GetExtension(StageRootLayerPath);
 
 	// If we have multiple materials *within this mesh* that want to be emitted to the same filepath we'll append
diff --git a/Engine/Plugins/Importers/USDImporter/Source/USDExporter/Private/SkeletalMeshExporterUSD.cpp b/Engine/Plugins/Importers/USDImporter/Source/USDExporter/Private/SkeletalMeshExporterUSD.cpp
--- a/Engine/Plugins/Importers/USDImporter/Source/USDExporter/Private/SkeletalMeshExporterUSD.cpp
+++ b/Engine/Plugins/Importers/USDImporter/Source/USDExporter/Private/SkeletalMeshExporterUSD.cpp
@@ -7,6 +7,7 @@
 #include "USDClassesModule.h"
 #include "USDConversionUtils.h"
 #include "USDExporterModule.h"
+#include "USDExportUtils.h"
 #include "USDLog.h"
 #include "USDObjectUtils.h"
 #include "USDOptionsWindow.h"
@@ -169,6 +170,8 @@ bool USkeletalMeshExporterUsd::ExportBinary(
 		return false;
 	}
 
+	UsdUnreal::ExportUtils::FUniquePathScope UniquePathScope;
+
 	// Get a simple GUID hash/identifier of our mesh
 	FSHA1 SHA1;
 	UE::SkeletalMeshExporterUSD::Private::HashSkeletalMesh(SkeletalMesh, SHA1);
@@ -321,12 +324,7 @@ bool USkeletalMeshExporterUsd::ExportBinary(
 	}
 	else
 	{
-		UnrealToUsd::ConvertSkeletalMeshToStaticMesh(
-			SkeletalMesh,
-			RootPrim,
-			UsdUtils::GetDefaultTimeCode(),
-			&AssetStage
-		);
+		UnrealToUsd::ConvertSkeletalMeshToStaticMesh(SkeletalMesh, RootPrim, UsdUtils::GetDefaultTimeCode(), &AssetStage);
 	}
 
 	if (UE::FUsdPrim AssetDefaultPrim = AssetStage.GetDefaultPrim())
diff --git a/Engine/Plugins/Importers/USDImporter/Source/USDExporter/Private/StaticMeshExporterUSD.cpp b/Engine/Plugins/Importers/USDImporter/Source/USDExporter/Private/StaticMeshExporterUSD.cpp
--- a/Engine/Plugins/Importers/USDImporter/Source/USDExporter/Private/StaticMeshExporterUSD.cpp
+++ b/Engine/Plugins/Importers/USDImporter/Source/USDExporter/Private/StaticMeshExporterUSD.cpp
@@ -7,6 +7,7 @@
 #include "USDClassesModule.h"
 #include "USDConversionUtils.h"
 #include "USDExporterModule.h"
+#include "USDExportUtils.h"
 #include "USDGeomMeshConversion.h"
 #include "USDLog.h"
 #include "USDObjectUtils.h"
@@ -174,6 +175,8 @@ bool UStaticMeshExporterUsd::ExportBinary(UObject* Object, const TCHAR* Type, FA
 		return false;
 	}
 
+	UsdUnreal::ExportUtils::FUniquePathScope UniquePathScope;
+
 	// Get a simple GUID hash/identifier of our mesh and options
 	FSHA1 SHA1;
 	UE::StaticMeshExporterUSD::Private::HashStaticMesh(StaticMesh, SHA1);
diff --git a/Engine/Plugins/Importers/USDImporter/Source/USDExporter/Private/USDConversionBlueprintContext.cpp b/Engine/Plugins/Importers/USDImporter/Source/USDExporter/Private/USDConversionBlueprintContext.cpp
--- a/Engine/Plugins/Importers/USDImporter/Source/USDExporter/Private/USDConversionBlueprintContext.cpp
+++ b/Engine/Plugins/Importers/USDImporter/Source/USDExporter/Private/USDConversionBlueprintContext.cpp
@@ -495,13 +495,7 @@ bool UUsdConversionBlueprintContext::ConvertLandscapeProxyActorMaterial(
 		}
 	}
 
-	bool bSuccess = UnrealToUsd::ConvertFlattenMaterial(
-		Actor->GetPathName() + TEXT("_BakedMaterial"),
-		LandscapeFlattenMaterial,
-		PropertiesToBake,
-		TexturesDir,
-		Prim
-	);
+	bool bSuccess = UnrealToUsd::ConvertFlattenMaterial(Actor->GetActorLabel(), LandscapeFlattenMaterial, PropertiesToBake, TexturesDir, Prim);
 
 	// FMaterialUtilities::ExportLandscapeMaterial always bakes WorldNormals, so we need to write to this material that it's meant to be using
 	// world-space normals. On import, we check for this and set the proper scalar parameter on our UsdPreviewSurface material to compensate for it.
diff --git a/Engine/Plugins/Importers/USDImporter/Source/USDExporter/Private/USDConversionBlueprintLibrary.cpp b/Engine/Plugins/Importers/USDImporter/Source/USDExporter/Private/USDConversionBlueprintLibrary.cpp
--- a/Engine/Plugins/Importers/USDImporter/Source/USDExporter/Private/USDConversionBlueprintLibrary.cpp
+++ b/Engine/Plugins/Importers/USDImporter/Source/USDExporter/Private/USDConversionBlueprintLibrary.cpp
@@ -8,6 +8,7 @@
 #include "USDClassesModule.h"
 #include "USDConversionUtils.h"
 #include "USDExporterModule.h"
+#include "USDExportUtils.h"
 #include "USDLayerUtils.h"
 #include "USDLog.h"
 #include "USDObjectUtils.h"
@@ -643,6 +644,21 @@ void UUsdConversionBlueprintLibrary::ResumeAnalyticsEvents()
 	IUsdClassesModule::ResumeAnalyticsEvents();
 }
 
+void UUsdConversionBlueprintLibrary::BeginUniquePathScope()
+{
+	UsdUnreal::ExportUtils::BeginUniquePathScope();
+}
+
+void UUsdConversionBlueprintLibrary::EndUniquePathScope()
+{
+	UsdUnreal::ExportUtils::EndUniquePathScope();
+}
+
+FString UUsdConversionBlueprintLibrary::GetUniqueFilePathForExport(const FString& DesiredPathWithExtension)
+{
+	return UsdUnreal::ExportUtils::GetUniqueFilePathForExport(DesiredPathWithExtension);
+}
+
 void UUsdConversionBlueprintLibrary::RemoveAllPrimSpecs(const FString& StageRootLayer, const FString& PrimPath, const FString& TargetLayer)
 {
 #if USE_USD_SDK
@@ -820,7 +836,8 @@ namespace UE::UsdConversionBlueprintLibrary::Private
 		FPropertyChangedEvent PostChangeEvent{
 			FUsdMetadataValue::StaticStruct()->FindPropertyByName(GET_MEMBER_NAME_CHECKED(FUsdMetadataValue, StringifiedValue)),
 			ChangeType,
-			{TopLevelObject}};
+			{TopLevelObject}
+		};
 		PostChangeEvent.MemberProperty = UUsdAssetUserData::StaticClass()->FindPropertyByName(
 			GET_MEMBER_NAME_CHECKED(UUsdAssetUserData, StageIdentifierToMetadata)
 		);
diff --git a/Engine/Plugins/Importers/USDImporter/Source/USDExporter/Public/USDConversionBlueprintLibrary.h b/Engine/Plugins/Importers/USDImporter/Source/USDExporter/Public/USDConversionBlueprintLibrary.h
--- a/Engine/Plugins/Importers/USDImporter/Source/USDExporter/Public/USDConversionBlueprintLibrary.h
+++ b/Engine/Plugins/Importers/USDImporter/Source/USDExporter/Public/USDConversionBlueprintLibrary.h
@@ -154,6 +154,35 @@ public:
 	UFUNCTION(BlueprintCallable, Category = "USD|Analytics")
 	static void ResumeAnalyticsEvents();
 
+	/**
+	 * Begins a UniquePathScope, incrementing the internal scope counter.
+	 *
+	 * During a UniquePathScope, all paths returned by GetUniqueFilePathForExport will be globally unique (i.e. it will
+	 * never return the same path twice).
+	 *
+	 * Opening a scope while another scope is opened has no effect but to increment the scope counter further.
+	 */
+	UFUNCTION(BlueprintCallable, Category = "USD|Exporting utils")
+	static void BeginUniquePathScope();
+
+	/**
+	 * Ends a UniquePathScope, decrementing the internal scope counter.
+	 *
+	 * If the internal scope counter reaches zero (i.e. all previously opened scopes are ended) this also clears the
+	 * cache of unique paths.
+	 */
+	UFUNCTION(BlueprintCallable, Category = "USD|Exporting utils")
+	static void EndUniquePathScope();
+
+	/**
+	 * If we're inside of a UniquePathScope, returns a sanitized (and potentially suffixed) path that is guaranteed
+	 * to not collide with any other path returned from this function during the UniquePathScope.
+	 *
+	 * If we're not inside of a UniquePathScope, returns the sanitized version of DesiredPathWithExtension.
+	 */
+	UFUNCTION(BlueprintCallable, Category = "USD|Exporting utils")
+	static FString GetUniqueFilePathForExport(const FString& DesiredPathWithExtension);
+
 	/**
 	 * Removes all the prim specs for Prim on the given Layer.
 	 *
diff --git a/Engine/Plugins/Runtime/USDCore/Source/USDUtilities/Private/USDExportUtils.cpp b/Engine/Plugins/Runtime/USDCore/Source/USDUtilities/Private/USDExportUtils.cpp
new file mode 100644
--- /dev/null
+++ b/Engine/Plugins/Runtime/USDCore/Source/USDUtilities/Private/USDExportUtils.cpp
@@ -0,0 +1,83 @@
+// Copyright Epic Games, Inc. All Rights Reserved.
+
+#include "USDExportUtils.h"
+
+#include "Misc/Paths.h"
+
+namespace UE::USDExportUtils::Private
+{
+	static int32 UniquePathScopeWhenNonZero = 0;
+	static TSet<FString> UniquePathsForExport;
+}
+
+void UsdUnreal::ExportUtils::BeginUniquePathScope()
+{
+	using namespace UE::USDExportUtils::Private;
+
+	++UniquePathScopeWhenNonZero;
+}
+
+void UsdUnreal::ExportUtils::EndUniquePathScope()
+{
+	using namespace UE::USDExportUtils::Private;
+
+	--UniquePathScopeWhenNonZero;
+
+	if (UniquePathScopeWhenNonZero == 0)
+	{
+		UniquePathsForExport.Reset();
+	}
+}
+
+UsdUnreal::ExportUtils::FUniquePathScope::FUniquePathScope()
+{
+	BeginUniquePathScope();
+}
+
+UsdUnreal::ExportUtils::FUniquePathScope::~FUniquePathScope()
+{
+	EndUniquePathScope();
+}
+
+FString UsdUnreal::ExportUtils::GetUniqueFilePathForExport(const FString& DesiredPathWithExtension)
+{
+	using namespace UE::USDExportUtils::Private;
+
+	FString SanitizedPath = DesiredPathWithExtension;
+	SanitizeFilePath(SanitizedPath);
+
+	if (UniquePathScopeWhenNonZero == 0)
+	{
+		// Not in a unique path scope --> Just return the path directly
+		return SanitizedPath;
+	}
+
+	if (!UniquePathsForExport.Contains(SanitizedPath))
+	{
+		UniquePathsForExport.Add(SanitizedPath);
+		return SanitizedPath;
+	}
+
+	FString Path;
+	FString FileName;
+	FString Extension;
+	FPaths::Split(SanitizedPath, Path, FileName, Extension);
+	FString Prefix = FPaths::Combine(Path, FileName);
+
+	int32 Index = 0;
+	FString Result = FString::Printf(TEXT("%s_%d.%s"), *Prefix, Index, *Extension);
+	while (UniquePathsForExport.Contains(Result))
+	{
+		Result = FString::Printf(TEXT("%s_%d.%s"), *Prefix, ++Index, *Extension);
+	}
+
+	UniquePathsForExport.Add(Result);
+	return Result;
+}
+
+void UsdUnreal::ExportUtils::SanitizeFilePath(FString& Path)
+{
+	FPaths::NormalizeFilename(Path);
+	FPaths::RemoveDuplicateSlashes(Path);
+	FPaths::CollapseRelativeDirectories(Path);
+}
diff --git a/Engine/Plugins/Runtime/USDCore/Source/USDUtilities/Private/USDShadeConversion.cpp b/Engine/Plugins/Runtime/USDCore/Source/USDUtilities/Private/USDShadeConversion.cpp
--- a/Engine/Plugins/Runtime/USDCore/Source/USDUtilities/Private/USDShadeConversion.cpp
+++ b/Engine/Plugins/Runtime/USDCore/Source/USDUtilities/Private/USDShadeConversion.cpp
@@ -12,6 +12,7 @@
 #include "USDAttributeUtils.h"
 #include "USDConversionUtils.h"
 #include "USDErrorUtils.h"
+#include "USDExportUtils.h"
 #include "USDLayerUtils.h"
 #include "USDLog.h"
 #include "USDMemory.h"
@@ -1405,7 +1406,7 @@ namespace UE::UsdShadeConversion::Private
 	// each channel
 	TMap<EMaterialProperty, FString> WriteTextures(
 		FBakedMaterialView& BakedSamples,
-		const FString& TextureNamePrefix,
+		const FString& MaterialName,
 		const FDirectoryPath& TexturesFolder
 	)
 	{
@@ -1415,6 +1416,14 @@ namespace UE::UsdShadeConversion::Private
 
 		TMap<EMaterialProperty, FString> WrittenTexturesPerChannel;
 
+		TSet<FString> UsedFilePaths;
+		auto MakeUniqueTexturePath = [&UsedFilePaths, &MaterialName, &TexturesFolder](const FString& PropertyName, const FString& ExtensionWithDot)
+		{
+			const FString TextureFileName = UsdUnreal::ObjectUtils::SanitizeObjectName(MaterialName + TEXT("_") + PropertyName);
+			const FString TextureFilePathPrefix = FPaths::Combine(TexturesFolder.Path, TextureFileName);
+			return UsdUnreal::ExportUtils::GetUniqueFilePathForExport(TextureFilePathPrefix + ExtensionWithDot);
+		};
+
 		// Write textures for HDR baked properties larger than 1x1
 		for (TPair<EMaterialProperty, TArray<FFloat16Color>*>& BakedDataPair : BakedSamples.HDRPropertyData)
 		{
@@ -1437,14 +1446,8 @@ namespace UE::UsdShadeConversion::Private
 			FString TrimmedPropertyName = PropertyName.ToString();
 			TrimmedPropertyName.RemoveFromStart(TEXT("MP_"));
 
-			// Final FilePath will be something like "C:/TexturesFolder/Game_ContentFolder_Materials_Red_BaseColor.png", which automatically
-			// guarantees it won't overwrite another texture from the same export, but will overwrite old textures from previous exports
-			FString TextureFileName = UsdUnreal::ObjectUtils::SanitizeObjectName(FPaths::ChangeExtension(TextureNamePrefix, TEXT("")));
-			TextureFileName.RemoveFromStart(TEXT("_"));
-			FString TextureFilePath = FPaths::Combine(
-				TexturesFolder.Path,
-				FString::Printf(TEXT("%s_%s.exr"), *UsdUnreal::ObjectUtils::SanitizeObjectName(TextureFileName), *TrimmedPropertyName)
-			);
+			const FString ExtensionWithDot = TEXT(".exr");
+			const FString TextureFilePath = MakeUniqueTexturePath(TrimmedPropertyName, ExtensionWithDot);
 
 			// For some reason the baked samples always have zero alpha and there is nothing we can do about it... It seems like the material
 			// baking module is made with the intent that the data ends up in UTexture2Ds, where they can be set to be compressed without
@@ -1499,14 +1502,8 @@ namespace UE::UsdShadeConversion::Private
 			FString TrimmedPropertyName = PropertyName.ToString();
 			TrimmedPropertyName.RemoveFromStart(TEXT("MP_"));
 
-			// Final FilePath will be something like "C:/TexturesFolder/Game_ContentFolder_Materials_Red_BaseColor.png", which automatically
-			// guarantees it won't overwrite another texture from the same export, but will overwrite old textures from previous exports
-			FString TextureFileName = UsdUnreal::ObjectUtils::SanitizeObjectName(FPaths::ChangeExtension(TextureNamePrefix, TEXT("")));
-			TextureFileName.RemoveFromStart(TEXT("_"));
-			FString TextureFilePath = FPaths::Combine(
-				TexturesFolder.Path,
-				FString::Printf(TEXT("%s_%s.png"), *UsdUnreal::ObjectUtils::SanitizeObjectName(TextureFileName), *TrimmedPropertyName)
-			);
+			const FString ExtensionWithDot = TEXT(".png");
+			const FString TextureFilePath = MakeUniqueTexturePath(TrimmedPropertyName, ExtensionWithDot);
 
 			// For some reason the baked samples always have zero alpha and there is nothing we can do about it... It seems like the material
 			// baking module is made with the intent that the data ends up in UTexture2Ds, where they can be set to be compressed without
@@ -2260,7 +2257,8 @@ namespace UE::UsdShadeConversion::Private
 				{
 					using MtlxType = decltype(Type);
 
-					// complex template to just retrieve the right type, basically a mapping between Usd and Mtlx for Vector types, otherwise we just take the type (float, int, string)
+					// complex template to just retrieve the right type, basically a mapping between Usd and Mtlx for Vector types, otherwise we just
+					// take the type (float, int, string)
 					// clang-format off
 					using UsdType = std::conditional_t<
 						std::is_same_v<MtlxType, mx::Color4> || std::is_same_v<MtlxType, mx::Vector4>,
@@ -2898,7 +2896,7 @@ bool UnrealToUsd::ConvertMaterialToBakedSurface(
 	}
 
 	UsdShadeConversionImpl::FBakedMaterialView View{BakedData};
-	TMap<EMaterialProperty, FString> WrittenTextures = UsdShadeConversionImpl::WriteTextures(View, InMaterial.GetPathName(), InTexturesDir);
+	TMap<EMaterialProperty, FString> WrittenTextures = UsdShadeConversionImpl::WriteTextures(View, InMaterial.GetName(), InTexturesDir);
 
 	// Manually add user supplied constant values. Can't place these in InMaterial as they're floats, and baked data is just quantized FColors
 	TMap<EMaterialProperty, float> UserConstantValues;
diff --git a/Engine/Plugins/Runtime/USDCore/Source/USDUtilities/Public/USDExportUtils.h b/Engine/Plugins/Runtime/USDCore/Source/USDUtilities/Public/USDExportUtils.h
new file mode 100644
--- /dev/null
+++ b/Engine/Plugins/Runtime/USDCore/Source/USDUtilities/Public/USDExportUtils.h
@@ -0,0 +1,53 @@
+// Copyright Epic Games, Inc. All Rights Reserved.
+
+#pragma once
+
+#include "CoreMinimal.h"
+
+namespace UsdUnreal::ExportUtils
+{
+	/**
+	 * Begins a UniquePathScope, incrementing the internal scope counter.
+	 *
+	 * During a UniquePathScope, all paths returned by GetUniqueFilePathForExport will be globally unique (i.e. it will
+	 * never return the same path twice).
+	 *
+	 * Opening a scope while another scope is opened has no effect but to increment the scope counter further.
+	 */
+	USDUTILITIES_API void BeginUniquePathScope();
+
+	/**
+	 * Ends a UniquePathScope, decrementing the internal scope counter.
+	 *
+	 * If the internal scope counter reaches zero (i.e. all previously opened scopes are ended) this also clears the
+	 * cache of unique paths.
+	 */
+	USDUTILITIES_API void EndUniquePathScope();
+
+	/** Utility to call BeginUniquePathScope() on construction and EndUniquePathScope() on destruction */
+	class USDUTILITIES_API FUniquePathScope final
+	{
+	public:
+		FUniquePathScope();
+		~FUniquePathScope();
+
+		FUniquePathScope(const FUniquePathScope&) = delete;
+		FUniquePathScope(FUniquePathScope&&) = delete;
+		FUniquePathScope& operator=(const FUniquePathScope&) = delete;
+		FUniquePathScope& operator=(FUniquePathScope&&) = delete;
+	};
+
+	/**
+	 * If we're inside of a UniquePathScope, returns a sanitized (and potentially suffixed) path that is guaranteed
+	 * to not collide with any other path returned from this function during the UniquePathScope.
+	 *
+	 * If we're not inside of a UniquePathScope, returns the sanitized version of DesiredPathWithExtension.
+	 */
+	USDUTILITIES_API FString GetUniqueFilePathForExport(const FString& DesiredPathWithExtension);
+
+	/**
+	 * Sanitizes the Path to be used as a clean absolute file path.
+	 * Normalizes separators, removes duplicate separators, collapses relative paths, etc.
+	 */
+	USDUTILITIES_API void SanitizeFilePath(FString& Path);
+}
