[Interchange - USD] Refactor skeletal mesh processing

#rnx

diff --git a/Engine/Plugins/Interchange/Runtime/Source/Import/Private/Usd/InterchangeUsdTranslator.cpp b/Engine/Plugins/Interchange/Runtime/Source/Import/Private/Usd/InterchangeUsdTranslator.cpp
--- a/Engine/Plugins/Interchange/Runtime/Source/Import/Private/Usd/InterchangeUsdTranslator.cpp
+++ b/Engine/Plugins/Interchange/Runtime/Source/Import/Private/Usd/InterchangeUsdTranslator.cpp
@@ -541,26 +541,6 @@ namespace UE::InterchangeUsdTranslator::Private
 		return true;
 	}
 
-	void FixSkeletalMeshDescriptionColors(FMeshDescription& MeshDescription)
-	{
-		// FSkeletalMeshImportData::GetMeshDescription() will reinterpret our Wedge FColors as linear, and put those
-		// sRGB values disguised as linear into the mesh description. This also seems to disagree with the patch on
-		// cl 32791826, so here we have to fix that up and get our mesh description colors to be actually linear...
-		//
-		// This will hopefully go away once we have our own skinned mesh to FMeshDescription conversion function.
-		//
-		// Note: Weirdly enough skeletal meshes seem to put linear colors on VertexColor output, while static meshes
-		// put sRGB colors? Maybe this is why the comment above the change on 32791826 mentions to remove the ToFColor on
-		// StaticMeshBuilder? This is overall very confusing
-		FStaticMeshAttributes Attributes(MeshDescription);
-		TVertexInstanceAttributesRef<FVector4f> VertexColor = Attributes.GetVertexInstanceColors();
-		for (FVertexInstanceID VertexInstanceID : MeshDescription.VertexInstances().GetElementIDs())
-		{
-			const FColor ActualSRGB = FLinearColor(VertexColor[VertexInstanceID]).ToFColor(false);
-			VertexColor[VertexInstanceID] = FLinearColor{ActualSRGB};
-		}
-	}
-
 	void FixMaterialSlotNames(FMeshDescription& MeshDescription, const TArray<UsdUtils::FUsdPrimMaterialSlot>& MeshAssingmentSlots)
 	{
 		// Fixup material slot names to match the material that is assigned. For Interchange it is better to have the material
@@ -2711,26 +2691,19 @@ namespace UE::InterchangeUsdTranslator::Private
 
 		FSkeletalMeshImportData SkelMeshImportData;
 		UsdUtils::FUsdPrimMaterialAssignmentInfo TempMaterialInfo;
-		bool bSuccess = UsdToUnreal::ConvertSkinnedMesh(SkinningQuery, SkelQuery, SkelMeshImportData, TempMaterialInfo, Options);
-		if (!bSuccess)
-		{
-			return false;
-		}
-
-		// TODO: Swap this code path with some function to directly convert a skinned USD mesh to MeshDescription.
-		// We need that on the other USD workflows as well, not only here...
-		//
-		// Note: This is also doubly bad because it internally recomputes tangents and normals, which will also
-		// be done by Interchange later..
-		const USkeletalMesh* UnusedSkelMesh = nullptr;
-		FSkeletalMeshBuildSettings* UnusedBuildSettings = nullptr;
 		FMeshDescription TempMeshDescription;
-		bSuccess = SkelMeshImportData.GetMeshDescription(UnusedSkelMesh, UnusedBuildSettings, TempMeshDescription);
+
+		pxr::UsdSkelSkinningQuery& PxrSkelSkinningQueryPtr = SkinningQuery;
+		pxr::UsdSkelSkeletonQuery& PxrSkelSkeletonQueryPtr = SkelQuery;
+
+		bool bSuccess = UsdToUnreal::ConvertGeomMesh(Prim, TempMeshDescription, TempMaterialInfo, Options, &PxrSkelSkinningQueryPtr, &PxrSkelSkeletonQueryPtr);
 		if (!bSuccess)
 		{
 			return false;
 		}
 
+		OutMeshDescription = MoveTemp(TempMeshDescription);
+
 		UsdUtils::FUsdPrimMaterialAssignmentInfo* AssignmentPtr = &TempMaterialInfo;
 		if (bIsLODMesh)
 		{
@@ -2749,10 +2722,6 @@ namespace UE::InterchangeUsdTranslator::Private
 			}
 		}
 
-		OutMeshDescription = MoveTemp(TempMeshDescription);
-
-		FixSkeletalMeshDescriptionColors(OutMeshDescription);
-
 		FixMaterialSlotNames(OutMeshDescription, AssignmentPtr->Slots);
 
 		OutJointNames = *JointNames;
diff --git a/Engine/Plugins/Runtime/USDCore/Source/USDUtilities/Private/USDGeomMeshConversion.cpp b/Engine/Plugins/Runtime/USDCore/Source/USDUtilities/Private/USDGeomMeshConversion.cpp
--- a/Engine/Plugins/Runtime/USDCore/Source/USDUtilities/Private/USDGeomMeshConversion.cpp
+++ b/Engine/Plugins/Runtime/USDCore/Source/USDUtilities/Private/USDGeomMeshConversion.cpp
@@ -36,6 +36,7 @@
 #include "MeshDescription.h"
 #include "Misc/Paths.h"
 #include "Rendering/SkeletalMeshRenderData.h"
+#include "SkeletalMeshAttributes.h"
 #include "SkeletalRenderPublic.h"
 #include "StaticMeshAttributes.h"
 #include "StaticMeshOperations.h"
@@ -75,6 +76,9 @@
 #include "pxr/usd/usdShade/tokens.h"
 #include "pxr/usd/usdSkel/bindingAPI.h"
 #include "pxr/usd/usdSkel/root.h"
+#include "pxr/usd/usdSkel/skeletonQuery.h"
+#include "pxr/usd/usdSkel/skinningQuery.h"
+#include "pxr/usd/usdSkel/utils.h"
 #include "pxr/usdImaging/usdImaging/implicitSurfaceMeshUtils.h"
 #include "USDIncludesEnd.h"
 
@@ -1122,6 +1126,13 @@ namespace UE::UsdGeomMeshConversion::Private
 		return UsdUtils::CombinePrimvarsIntoUVSets(AllPrimvars, PreferredPrimvars);
 	}
 
+	struct FUsdRawBoneInfluence // just weight, vertex, and Bone, sorted in CopySkinWeightsToAttribute
+	{
+		float Weight;
+		int32 VertexIndex;
+		int32 BoneIndex;
+	};
+
 	// Unconverted, raw USD mesh data to convert into a MeshDescription
 	struct FUsdMeshData
 	{
@@ -1173,13 +1184,99 @@ namespace UE::UsdGeomMeshConversion::Private
 
 		TOptional<int32> ProvidedNumUVSets;
 		int32 MaterialIndexOffset = 0;
+
+		//Attributes for SkeletalMesh:
+		TArray <FUsdRawBoneInfluence> Influences;
 	};
 
+	namespace USDSkeletalDataConversion
+	{
+		bool HasMultipleRootBones(const pxr::UsdSkelSkeletonQuery& SkeletonQuery)
+		{
+			const pxr::UsdSkelTopology& SkelTopology = SkeletonQuery.GetTopology();
+			const pxr::VtArray<int>& JointParentIndices = SkelTopology.GetParentIndices();
+			bool bFoundRoot = false;
+			for (int ParentIndex : JointParentIndices)
+			{
+				if (ParentIndex == INDEX_NONE)
+				{
+					if (bFoundRoot)
+					{
+						return true;
+					}
+					bFoundRoot = true;
+				}
+			}
+
+			return false;
+		}
+
+		// Based on SkeletalMeshLODImporterData.cpp 's CopySkinWeightsToAttribute.
+		void CopySkinWeightsToAttribute(
+			const TArray<FUsdRawBoneInfluence>& InInfluences,
+			const TArray<FVertexID>& InVertexIDMap,
+			FSkinWeightsVertexAttributesRef OutSkinWeightsAttribute
+		)
+		{
+			using namespace UE::AnimationCore;
+
+			// The weights are stored with links back to the vertices, rather than being compact.
+			// Make a copy of the weights, sort them by vertex id and go by equal vertex-id strides.
+			// We could do an indirection but the traversal + setup cost is probably not worth it.
+			TArray<FUsdRawBoneInfluence> SortedInfluences(InInfluences);
+			SortedInfluences.Sort([](const FUsdRawBoneInfluence& A, const FUsdRawBoneInfluence& B)
+				{
+					return A.VertexIndex < B.VertexIndex;
+				});
+
+			// Do the base skin weights first. We do the alternate skin weights later, since they may require geometric remapping. 
+			TArray<FBoneWeight> BoneWeights;
+			for (int32 StartStride = 0, EndStride = 0; EndStride != SortedInfluences.Num(); StartStride = EndStride)
+			{
+				const int32 VertexIndex = SortedInfluences[StartStride].VertexIndex;
+
+				// There exist meshes where the influence map got auto-filled with 100% weight on root in by using the wedge count of the raw mesh,
+				// due to missing weights (e.g. static mesh imported as a skeletal mesh), and so may refer to vertices that don't exist.
+				// We just stop when we get to the broken set and ignore the rest.
+				if (VertexIndex >= InVertexIDMap.Num())
+				{
+					break;
+				}
+
+				EndStride = StartStride + 1;
+				while (EndStride < SortedInfluences.Num() && VertexIndex == SortedInfluences[EndStride].VertexIndex)
+				{
+					EndStride++;
+				}
+
+				BoneWeights.Reset(0);
+				for (int32 Idx = StartStride; Idx < EndStride; Idx++)
+				{
+					const FUsdRawBoneInfluence& RawInfluence = SortedInfluences[Idx];
+					int32 BoneIndex = RawInfluence.BoneIndex;
+
+					FBoneWeight BoneWeight(static_cast<FBoneIndexType>(BoneIndex), RawInfluence.Weight);
+					BoneWeights.Add(BoneWeight);
+				}
+
+				if (BoneWeights.IsEmpty())
+				{
+					const FBoneWeight RootBoneWeight(0, 1.0f);
+					BoneWeights.Add(RootBoneWeight);
+				}
+
+				OutSkinWeightsAttribute.Set(InVertexIDMap[VertexIndex], BoneWeights);
+			}
+		}
+	}
+
 	bool CollectMeshData(
 		const pxr::UsdPrim& Prim,
 		const UsdToUnreal::FUsdMeshConversionOptions& Options,
 		FUsdMeshData& InOutMeshData,
-		UsdUtils::FUsdPrimMaterialAssignmentInfo& InOutCombinedMaterialAssignments
+		UsdUtils::FUsdPrimMaterialAssignmentInfo& InOutCombinedMaterialAssignments,
+		const pxr::UsdSkelSkinningQuery* SkinningQuery = nullptr,
+		const pxr::UsdSkelSkeletonQuery* SkeletonQuery = nullptr
 	)
 	{
 		FScopedUsdAllocs Allocs;
@@ -1190,6 +1287,36 @@ namespace UE::UsdGeomMeshConversion::Private
 			return false;
 		}
 
+		const bool bProcessSkeletalAttributes = SkeletonQuery && SkeletonQuery->IsValid() && SkinningQuery && SkinningQuery->IsValid();
+
+		pxr::VtArray<pxr::GfMatrix4d> MeshToSkeletonRestPose;
+		if (bProcessSkeletalAttributes)
+		{
+			bool bSuccess = true;
+
+			// Get world-space restTransforms
+			pxr::VtArray<pxr::GfMatrix4d> WorldSpaceRestTransforms;
+			{
+				pxr::VtArray<pxr::GfMatrix4d> JointLocalRestTransforms;
+				const bool bAtRest = true;
+				bSuccess &= SkeletonQuery->ComputeJointLocalTransforms(&JointLocalRestTransforms, Options.TimeCode, bAtRest);
+				bSuccess &= pxr::UsdSkelConcatJointTransforms(SkeletonQuery->GetTopology(), JointLocalRestTransforms, &WorldSpaceRestTransforms);
+			}
+
+			// Get world-space bindTransforms
+			pxr::VtArray<pxr::GfMatrix4d> WorldSpaceBindTransforms;
+			bSuccess &= SkeletonQuery->GetJointWorldBindTransforms(&WorldSpaceBindTransforms);
+
+			if (bSuccess)
+			{
+				MeshToSkeletonRestPose.resize(WorldSpaceRestTransforms.size());
+				for (uint32 Index = 0; Index < WorldSpaceRestTransforms.size(); ++Index)
+				{
+					MeshToSkeletonRestPose[Index] = WorldSpaceBindTransforms[Index].GetInverse() * WorldSpaceRestTransforms[Index];
+				}
+			}
+		}
+
 		// All pointsBased/Gprim attributes we'll retrieve happen to have default varying interpolation
 		const pxr::TfToken DefaultInterpolation = pxr::UsdGeomTokens->varying;
 
@@ -1235,6 +1362,11 @@ namespace UE::UsdGeomMeshConversion::Private
 				PointsPrimvar.Get(&InOutMeshData.Points, Options.TimeCode);
 				InOutMeshData.PointInterpolation = GetAttrInterpolation(PointsPrimvar, pxr::UsdGeomTokens->vertex);
 				PointsPrimvar.GetIndices(&InOutMeshData.PointIndices, Options.TimeCode);
+
+				if (bProcessSkeletalAttributes)
+				{
+					SkinningQuery->ComputeSkinnedPoints(MeshToSkeletonRestPose, &InOutMeshData.Points, Options.TimeCode);
+				}
 			}
 			else if (pxr::UsdAttribute PointsAttr = UsdMesh.GetPointsAttr())
 			{
@@ -1374,6 +1506,99 @@ namespace UE::UsdGeomMeshConversion::Private
 			}
 		}
 
+		//Joint Indices/Weights
+		if (bProcessSkeletalAttributes)
+		{
+			const pxr::UsdPrim& SkinningPrim = SkinningQuery->GetPrim();
+			pxr::UsdSkelBindingAPI SkelBindingAPI(SkinningPrim);
+			if (SkelBindingAPI)
+			{
+				// ComputeJointInfluences returns the influences per bone that applies to all the points of the mesh
+				// ComputeVaryingJointInfluences returns the joint influences for each points, expanding the influences to all points if the mesh is rigidly
+				// deformed
+				pxr::VtArray<int> JointIndices;
+				pxr::VtArray<float> JointWeights;
+				size_t NumPoints = InOutMeshData.Points.size();
+				SkinningQuery->ComputeVaryingJointInfluences(NumPoints, &JointIndices, &JointWeights);
+
+				// Keep track of whether we added an additional "true" root bone in the cases the bound skeleton has
+				// multiple root bones
+				// We'll only ever set NumAdditionalBones to 1 or 0 (as we'll only either need a "true root bone" or
+				// not), but naming it this way allows us to use it like an offset, which should make it easier to
+				// understand whenever it is used
+				uint32 NumAdditionalBones = USDSkeletalDataConversion::HasMultipleRootBones(*SkeletonQuery) ? 1 : 0;
+				uint32 NumInfluencesPerComponent = SkinningQuery->GetNumInfluencesPerComponent();
+
+				// We keep track of which influences we added because we combine many Mesh prim (each with potentially a different
+				// explicit joint order) into the same skeletal mesh asset
+				if (JointWeights.size() > (NumPoints - 1) * (NumInfluencesPerComponent - 1))
+				{
+					uint32 JointIndex = 0;
+					InOutMeshData.Influences.Reserve(NumPoints);
+					for (uint32 PointIndex = 0; PointIndex < NumPoints; ++PointIndex)
+					{
+						// The JointIndices/JointWeights contain the influences data for NumPoints * NumInfluencesPerComponent
+						for (uint32 InfluenceIndex = 0; InfluenceIndex < NumInfluencesPerComponent; ++InfluenceIndex, ++JointIndex)
+						{
+							// BoneWeight could be 0 if the actual number of influences were less than NumInfluencesPerComponent for a given point so just ignore
+							// it
+							float BoneWeight = JointWeights[JointIndex];
+							if (BoneWeight != 0.f)
+							{
+								InOutMeshData.Influences.AddUninitialized();
+								InOutMeshData.Influences.Last().BoneIndex = NumAdditionalBones + JointIndices[JointIndex];
+								InOutMeshData.Influences.Last().Weight = BoneWeight;
+								InOutMeshData.Influences.Last().VertexIndex = PointIndex;
+							}
+						}
+					}
+				}
+				const int32 NumInfluences = InOutMeshData.Influences.Num();
+
+				// If we have a joint mapper this Mesh has an explicit joint ordering, so we need to map joint indices to the skeleton's bone indices
+				if (pxr::UsdSkelAnimMapperRefPtr AnimMapper = SkinningQuery->GetJointMapper())
+				{
+					pxr::VtArray<int> SkeletonBoneIndices;
+					if (pxr::UsdSkelSkeleton BoundSkeleton = SkelBindingAPI.GetInheritedSkeleton())
+					{
+						if (pxr::UsdAttribute SkeletonJointsAttr = BoundSkeleton.GetJointsAttr())
+						{
+							pxr::VtArray<pxr::TfToken> SkeletonJoints;
+							if (SkeletonJointsAttr.Get(&SkeletonJoints))
+							{
+								// If the skeleton has N bones, this will just contain { 0, 1, 2, ..., N-1 }
+								int NumUsdSkeletonBones = static_cast<int>(SkeletonJoints.size());
+								for (int SkeletonBoneIndex = 0; SkeletonBoneIndex < NumUsdSkeletonBones; ++SkeletonBoneIndex)
+								{
+									SkeletonBoneIndices.push_back(SkeletonBoneIndex);
+								}
+
+								// Use the AnimMapper to produce the indices of the Mesh's joints within the Skeleton's list of joints.
+								// Example: Imagine skeleton had { "Root", "Root/Hip", "Root/Hip/Shoulder", "Root/Hip/Shoulder/Arm", "Root/Hip/Shoulder/Arm/Elbow"
+								// }, and so BoneIndexRemapping was { 0, 1, 2, 3, 4 }. Consider a Mesh that specifies the explicit joints { "Root/Hip/Shoulder",
+								// "Root/Hip/Shoulder/Arm" }, and so uses the indices 0 and 1 to refer to Shoulder and Arm. After the Remap call
+								// SkeletonBoneIndices will hold { 2, 3 }, as those are the indices of Shoulder and Arm within the skeleton's bones
+								pxr::VtArray<int> BoneIndexRemapping;
+								if (AnimMapper->Remap(SkeletonBoneIndices, &BoneIndexRemapping))
+								{
+									for (int32 AddedInfluenceIndex = 0; AddedInfluenceIndex < NumInfluences; ++AddedInfluenceIndex)
+									{
+										FUsdRawBoneInfluence& Influence = InOutMeshData.Influences[AddedInfluenceIndex];
+
+										// We have to remove our "NumAdditionalBones" offset from the influence's bone index because that's a UE concept that
+										// the BoneIndexRemapping array doesn't really know about. After that, we have a bone index that matches the USD Skeleton
+										// joint order, then we can remap with BoneIndexRemapping and add our NumAdditionalBones back in so that it matches our
+										// USkeleton
+										Influence.BoneIndex = NumAdditionalBones + BoneIndexRemapping[Influence.BoneIndex - NumAdditionalBones];
+									}
+								}
+							}
+						}
+					}
+				}
+			}
+		}
+
 		// Material assignments
 		{
 			const bool bProvideMaterialIndices = true;
@@ -2264,6 +2489,7 @@ namespace UE::UsdGeomMeshConversion::Private
 
 		// Vertex positions
 		TVertexAttributesRef<FVector3f> MeshDescriptionVertexPositions = StaticMeshAttributes.GetVertexPositions();
+		TArray<FVertexID> VertexIDMap;
 		{
 			if (InMeshData.Points.size() < 3)
 			{
@@ -2271,6 +2497,7 @@ namespace UE::UsdGeomMeshConversion::Private
 			}
 
 			OutMeshDescription.ReserveNewVertices(InMeshData.Points.size());
+			VertexIDMap.Reserve(InMeshData.Points.size());
 
 			for (size_t LocalPointIndex = 0; LocalPointIndex < InMeshData.Points.size(); ++LocalPointIndex)
 			{
@@ -2279,6 +2506,7 @@ namespace UE::UsdGeomMeshConversion::Private
 				FVector Position = InOptions.AdditionalTransform.TransformPosition(UsdToUnreal::ConvertVector(InStageInfo, Point));
 
 				FVertexID AddedVertexId = OutMeshDescription.CreateVertex();
+				VertexIDMap.Add(AddedVertexId);
 				MeshDescriptionVertexPositions[AddedVertexId] = (FVector3f)Position;
 			}
 		}
@@ -2623,6 +2851,14 @@ namespace UE::UsdGeomMeshConversion::Private
 			}
 		}
 
+		if (InMeshData.Influences.Num())
+		{
+			FSkeletalMeshAttributes MeshAttributes(OutMeshDescription);
+			MeshAttributes.Register(bKeepExistingAttribute);
+			FSkinWeightsVertexAttributesRef VertexSkinWeights = MeshAttributes.GetVertexSkinWeights();
+			USDSkeletalDataConversion::CopySkinWeightsToAttribute(InMeshData.Influences, VertexIDMap, VertexSkinWeights);
+		}
+
 		if (NumPolygons > 0 && NumSkippedPolygons > 0)
 		{
 			UE_LOG(
@@ -2660,7 +2896,9 @@ bool UsdToUnreal::ConvertGeomMesh(
 	const pxr::UsdGeomMesh& UsdMesh,
 	FMeshDescription& OutMeshDescription,
 	UsdUtils::FUsdPrimMaterialAssignmentInfo& OutMaterialAssignments,
-	const FUsdMeshConversionOptions& Options
+	const FUsdMeshConversionOptions& Options,
+	const pxr::UsdSkelSkinningQuery* SkinningQuery,
+	const pxr::UsdSkelSkeletonQuery* SkeletonQuery
 )
 {
 	TRACE_CPUPROFILER_EVENT_SCOPE(UsdToUnreal::ConvertGeomMesh);
@@ -2676,7 +2914,7 @@ bool UsdToUnreal::ConvertGeomMesh(
 
 	UsdGeomMeshImpl::FUsdMeshData MeshData;
 
-	UsdGeomMeshImpl::CollectMeshData(UsdPrim, Options, MeshData, OutMaterialAssignments);
+	UsdGeomMeshImpl::CollectMeshData(UsdPrim, Options, MeshData, OutMaterialAssignments, SkinningQuery, SkeletonQuery);
 
 	if (Options.SubdivisionLevel > 0 && MeshData.SubdivScheme != pxr::UsdGeomTokens->none)
 	{
@@ -2699,7 +2937,9 @@ bool UsdToUnreal::ConvertGeomMesh(
 	const pxr::UsdPrim& MeshPrim,
 	FMeshDescription& InOutMeshDescription,
 	UsdUtils::FUsdPrimMaterialAssignmentInfo& InOutMaterialAssignments,
-	const FUsdMeshConversionOptions& CommonOptions
+	const FUsdMeshConversionOptions& CommonOptions,
+	const pxr::UsdSkelSkinningQuery* SkinningQuery,
+	const pxr::UsdSkelSkeletonQuery* SkeletonQuery
 )
 {
 	pxr::UsdGeomMesh GeomMesh{MeshPrim};
@@ -2708,7 +2948,7 @@ bool UsdToUnreal::ConvertGeomMesh(
 		return false;
 	}
 
-	return ConvertGeomMesh(GeomMesh, InOutMeshDescription, InOutMaterialAssignments, CommonOptions);
+	return ConvertGeomMesh(GeomMesh, InOutMeshDescription, InOutMaterialAssignments, CommonOptions, SkinningQuery, SkeletonQuery);
 }
 
 bool UsdToUnreal::ConvertPointInstancerToMesh(
diff --git a/Engine/Plugins/Runtime/USDCore/Source/USDUtilities/Public/USDGeomMeshConversion.h b/Engine/Plugins/Runtime/USDCore/Source/USDUtilities/Public/USDGeomMeshConversion.h
--- a/Engine/Plugins/Runtime/USDCore/Source/USDUtilities/Public/USDGeomMeshConversion.h
+++ b/Engine/Plugins/Runtime/USDCore/Source/USDUtilities/Public/USDGeomMeshConversion.h
@@ -22,6 +22,8 @@ PXR_NAMESPACE_OPEN_SCOPE
 	class UsdPrim;
 	class UsdShadeMaterial;
 	class UsdTyped;
+	class UsdSkelSkinningQuery;
+	class UsdSkelSkeletonQuery;
 PXR_NAMESPACE_CLOSE_SCOPE
 #endif	  // USE_USD_SDK
 
@@ -104,24 +106,31 @@ namespace UsdToUnreal
 	 * If a PrimvarToUVIndex mapping is provided in InOutMaterialAssignments, it will be used to remap the read
 	 * primvars into UV sets. Otherwise, a new PrimvarToUVIndex map will be constructed on-demand and assigned
 	 * to InOutMaterialAssignments instead.
+	 * If SkinningQuery and SkeletonQuery are provided we parse the Skeletal Attributes and fill the FMeshDescription accordingly.
 	 *
 	 * @param Mesh - Mesh prim to convert
 	 * @param InOutMeshDescription - Input/Output parameter that will be filled with the converted data
 	 * @param InOutMaterialAssignments - Input/Output parameter that will be filled with the material data extracted from Mesh
 	 * @param CommonOptions - Shared options used for the conversion
+	 * @param SkinningQuery - For computing Skinned Points, if the target and source data set is for a SkeletalMesh.
+	 * @param SkeletonQuery - For computing Skinned Points, if the target and source data set is for a SkeletalMesh. (primarily for MeshToSkeletonRestPose calculations)
 	 * @return Whether the conversion was successful or not.
 	 */
 	USDUTILITIES_API bool ConvertGeomMesh(
 		const pxr::UsdGeomMesh& Mesh,
 		FMeshDescription& InOutMeshDescription,
 		UsdUtils::FUsdPrimMaterialAssignmentInfo& InOutMaterialAssignments,
-		const FUsdMeshConversionOptions& CommonOptions = FUsdMeshConversionOptions::DefaultOptions
+		const FUsdMeshConversionOptions& CommonOptions = FUsdMeshConversionOptions::DefaultOptions,
+		const pxr::UsdSkelSkinningQuery* SkinningQuery = nullptr,
+		const pxr::UsdSkelSkeletonQuery* SkeletonQuery = nullptr
 	);
 	USDUTILITIES_API bool ConvertGeomMesh(
 		const pxr::UsdPrim& MeshPrim,
 		FMeshDescription& InOutMeshDescription,
 		UsdUtils::FUsdPrimMaterialAssignmentInfo& InOutMaterialAssignments,
-		const FUsdMeshConversionOptions& CommonOptions = FUsdMeshConversionOptions::DefaultOptions
+		const FUsdMeshConversionOptions& CommonOptions = FUsdMeshConversionOptions::DefaultOptions,
+		const pxr::UsdSkelSkinningQuery* SkinningQuery = nullptr,
+		const pxr::UsdSkelSkeletonQuery* SkeletonQuery = nullptr
 	);
 
 	/** Bake a point instancer prim into a single FMeshDescription and MaterialAssignments, manually "instancing" all
diff --git a/Engine/Plugins/Runtime/USDCore/Source/USDUtilities/USDUtilities.Build.cs b/Engine/Plugins/Runtime/USDCore/Source/USDUtilities/USDUtilities.Build.cs
--- a/Engine/Plugins/Runtime/USDCore/Source/USDUtilities/USDUtilities.Build.cs
+++ b/Engine/Plugins/Runtime/USDCore/Source/USDUtilities/USDUtilities.Build.cs
@@ -43,6 +43,7 @@ namespace UnrealBuildTool.Rules
 					"Slate",
 					"SlateCore",
 					"StaticMeshDescription",
+					"SkeletalMeshDescription"
 				}
 			);
 
