[USD] Refactored SkeletalMesh processing's SkinningPoint calculation fixes.

#rnx

diff --git a/Engine/Plugins/Runtime/USDCore/Source/USDUtilities/Private/USDGeomMeshConversion.cpp b/Engine/Plugins/Runtime/USDCore/Source/USDUtilities/Private/USDGeomMeshConversion.cpp
--- a/Engine/Plugins/Runtime/USDCore/Source/USDUtilities/Private/USDGeomMeshConversion.cpp
+++ b/Engine/Plugins/Runtime/USDCore/Source/USDUtilities/Private/USDGeomMeshConversion.cpp
@@ -1362,11 +1362,6 @@ namespace UE::UsdGeomMeshConversion::Private
 				PointsPrimvar.Get(&InOutMeshData.Points, Options.TimeCode);
 				InOutMeshData.PointInterpolation = GetAttrInterpolation(PointsPrimvar, pxr::UsdGeomTokens->vertex);
 				PointsPrimvar.GetIndices(&InOutMeshData.PointIndices, Options.TimeCode);
-
-				if (bProcessSkeletalAttributes)
-				{
-					SkinningQuery->ComputeSkinnedPoints(MeshToSkeletonRestPose, &InOutMeshData.Points, Options.TimeCode);
-				}
 			}
 			else if (pxr::UsdAttribute PointsAttr = UsdMesh.GetPointsAttr())
 			{
@@ -1374,6 +1369,12 @@ namespace UE::UsdGeomMeshConversion::Private
 				InOutMeshData.PointInterpolation = GetAttrInterpolation(PointsAttr, pxr::UsdGeomTokens->vertex);
 			}
 
+			//Apply Skinning if applicable.
+			if (bProcessSkeletalAttributes)
+			{
+				SkinningQuery->ComputeSkinnedPoints(MeshToSkeletonRestPose, &InOutMeshData.Points, Options.TimeCode);
+			}
+
 			// Normals
 			if (pxr::UsdGeomPrimvar NormalsPrimvar = pxr::UsdGeomPrimvar(Prim.GetAttribute(UnrealIdentifiers::PrimvarsNormals)))
 			{
