USD: Translate USD attributes as Interchange translated node attributes (just non-animated, non-array for now).

diff --git a/Engine/Plugins/Interchange/Runtime/Source/Common/Public/Usd/InterchangeUsdDefinitions.h b/Engine/Plugins/Interchange/Runtime/Source/Common/Public/Usd/InterchangeUsdDefinitions.h
--- a/Engine/Plugins/Interchange/Runtime/Source/Common/Public/Usd/InterchangeUsdDefinitions.h
+++ b/Engine/Plugins/Interchange/Runtime/Source/Common/Public/Usd/InterchangeUsdDefinitions.h
@@ -11,9 +11,6 @@ namespace UE::Interchange::USD
 	// Name of a custom attribute added to translated nodes to contain their geometry purpose (proxy, render, guide, etc.)
 	const FString GeometryPurposeIdentifier = TEXT("USD_Geometry_Purpose");
 
-	// Prim kind metadata
-	const FString KindIdentifier = TEXT("USD_Kind");
-
 	// Prefix for custom attributes that describe material instance parameters
 	const FString AttributeIdentifierPrefix = TEXT("USD_MI_Attribute_");
 
diff --git a/Engine/Plugins/Interchange/Runtime/Source/Import/Private/Usd/InterchangeUsdTranslator.cpp b/Engine/Plugins/Interchange/Runtime/Source/Import/Private/Usd/InterchangeUsdTranslator.cpp
--- a/Engine/Plugins/Interchange/Runtime/Source/Import/Private/Usd/InterchangeUsdTranslator.cpp
+++ b/Engine/Plugins/Interchange/Runtime/Source/Import/Private/Usd/InterchangeUsdTranslator.cpp
@@ -47,9 +47,11 @@
 #include "InterchangeShaderGraphNode.h"
 #include "InterchangeTexture2DNode.h"
 #include "InterchangeTranslatorHelper.h"
+#include "Internationalization/Regex.h"
 #include "Mesh/InterchangeMeshPayload.h"
 #include "Modules/ModuleManager.h"
 #include "MovieSceneSection.h"
+#include "Nodes/InterchangeUserDefinedAttribute.h"
 #include "Rendering/SkeletalMeshLODImporterData.h"
 #include "StaticMeshAttributes.h"
 #include "UDIMUtilities.h"
@@ -2386,6 +2388,114 @@ namespace UE::InterchangeUsdTranslator::Private
 #endif
 	}
 
+	template<typename T>
+	void TranslateAttribute(const UE::FUsdAttribute& Attr, const FString& AttrName, UInterchangeSceneNode* Node)
+	{
+		T Value;
+		if (Attr.Get<T>(Value))	   // Always check for an opinion on the default time code
+		{
+			TOptional<FString> PayloadKey;
+			UInterchangeUserDefinedAttributesAPI::CreateUserDefinedAttribute(Node, AttrName, Value, PayloadKey);
+		}
+	};
+
+	void TranslateAttributes(const UE::FUsdPrim& Prim, UInterchangeSceneNode* Node, const FString& AllowedAttributeRegex)
+	{
+		FRegexPattern RegexPattern{AllowedAttributeRegex};
+
+		for (const UE::FUsdAttribute& Attr : Prim.GetAttributes())
+		{
+			if (!Attr.HasAuthoredValue())
+			{
+				continue;
+			}
+
+			const FString AttrName = Attr.GetName().ToString();
+
+			FRegexMatcher RegexMatcher{RegexPattern, AttrName};
+			if (!RegexMatcher.FindNext())
+			{
+				continue;
+			}
+
+			// Note: We could have used Cpp type names here instead of the value type token, but doing the
+			// latter lets us handle the color types as FLinearColor, which is probably more useful
+			// References:
+			// - https://openusd.org/docs/api/_usd__page__datatypes.html
+			// - EAttributeTypes declaration
+			// - Engine/Source/Runtime/Interchange/Core/Private/Tests/StorageTest.cpp
+			// clang-format off
+			using TranslationFunc = TFunction<void(const UE::FUsdAttribute&, const FString&, UInterchangeSceneNode*)>;
+			static TMap<FName, TranslationFunc> TranslationFuncs = {
+				{TEXT("bool"),  		TranslateAttribute<bool>},
+
+				{TEXT("uchar"), 		TranslateAttribute<uint8>},
+				{TEXT("uchar[]"), 		TranslateAttribute<TArray<uint8>>},
+				{TEXT("int"), 			TranslateAttribute<int32>},
+				{TEXT("uint"), 			TranslateAttribute<uint32>},
+				{TEXT("int64"), 		TranslateAttribute<int64>},
+				{TEXT("uint64"), 		TranslateAttribute<uint64>},
+
+				{TEXT("half"),	 		TranslateAttribute<FFloat16>},
+				{TEXT("float"), 		TranslateAttribute<float>},
+				{TEXT("double"), 		TranslateAttribute<double>},
+				{TEXT("timecode"), 		TranslateAttribute<double>},
+
+				{TEXT("string"), 		TranslateAttribute<FString>},
+				{TEXT("token"), 		TranslateAttribute<FString>},
+				{TEXT("asset"), 		TranslateAttribute<FString>},
+
+				{TEXT("matrix2d"), 		TranslateAttribute<FMatrix44d>}, // Interchange only supports 4x4 matrices
+				{TEXT("matrix3d"), 		TranslateAttribute<FMatrix44d>},
+				{TEXT("matrix4d"), 		TranslateAttribute<FMatrix44d>},
+				{TEXT("frame4d"), 		TranslateAttribute<FMatrix44d>},
+
+				{TEXT("quath"), 		TranslateAttribute<FQuat4f>}, // No analogue type for Quat of halfs
+				{TEXT("quatf"), 		TranslateAttribute<FQuat4f>},
+				{TEXT("quatd"), 		TranslateAttribute<FQuat4d>},
+
+				{TEXT("half2"), 		TranslateAttribute<FVector2DHalf>},
+				{TEXT("float2"), 		TranslateAttribute<FVector2f>},
+				{TEXT("double2"), 		TranslateAttribute<FVector2d>},
+				{TEXT("int2"), 			TranslateAttribute<FIntPoint>},
+
+				{TEXT("half3"), 		TranslateAttribute<FVector3f>}, // There is no FVector3DHalf
+				{TEXT("point3h"), 		TranslateAttribute<FVector3f>},
+				{TEXT("normal3h"), 		TranslateAttribute<FVector3f>},
+				{TEXT("vector3h"), 		TranslateAttribute<FVector3f>},
+				{TEXT("color3h"), 		TranslateAttribute<FLinearColor>},
+
+				{TEXT("float3"), 		TranslateAttribute<FVector3f>},
+				{TEXT("point3f"), 		TranslateAttribute<FVector3f>},
+				{TEXT("normal3f"), 		TranslateAttribute<FVector3f>},
+				{TEXT("vector3f"), 		TranslateAttribute<FVector3f>},
+				{TEXT("color3f"), 		TranslateAttribute<FLinearColor>},
+
+				{TEXT("double3"), 		TranslateAttribute<FVector3d>},
+				{TEXT("point3d"), 		TranslateAttribute<FVector3d>},
+				{TEXT("normal3d"), 		TranslateAttribute<FVector3d>},
+				{TEXT("vector3d"), 		TranslateAttribute<FVector3d>},
+				{TEXT("color3d"), 		TranslateAttribute<FLinearColor>},
+
+				{TEXT("int3"), 			TranslateAttribute<FIntVector>},
+
+				{TEXT("half4"), 		TranslateAttribute<FVector4f>}, // There is no FVector4DHalf
+				{TEXT("float4"), 		TranslateAttribute<FVector4f>},
+				{TEXT("double4"), 		TranslateAttribute<FVector4d>},
+				{TEXT("int4"), 			TranslateAttribute<FIntRect>},
+				{TEXT("color4h"), 		TranslateAttribute<FLinearColor>},
+				{TEXT("color4f"), 		TranslateAttribute<FLinearColor>},
+				{TEXT("color4d"), 		TranslateAttribute<FLinearColor>},
+			};
+			// clang-format on
+
+			if (TranslationFunc Func = TranslationFuncs.FindRef(Attr.GetTypeName()))
+			{
+				Func(Attr, AttrName, Node);
+			}
+		}
+	}
+
 	void Traverse(
 		const UE::FUsdPrim& Prim,
 		UInterchangeUSDTranslatorImpl& TranslatorImpl,
@@ -2460,7 +2570,14 @@ namespace UE::InterchangeUsdTranslator::Private
 			const FString KindString = UsdToUnreal::ConvertToken(IUsdPrim::GetKind(Prim));
 			if (!KindString.IsEmpty())
 			{
-				SceneNode->AddStringAttribute(UE::Interchange::USD::KindIdentifier, KindString);
+				const FString UserDefinedAttributeName = TEXT("kind");
+				TOptional<FString> PayloadKey;
+				UInterchangeUserDefinedAttributesAPI::CreateUserDefinedAttribute(SceneNode, UserDefinedAttributeName, KindString, PayloadKey);
+			}
+
+			if (TranslatorSettings->bTranslatePrimAttributes)
+			{
+				TranslateAttributes(Prim, SceneNode, TranslatorSettings->AttributeRegexFilter);
 			}
 
 			// If we're an Xformable, get our transform.
@@ -2687,7 +2804,14 @@ namespace UE::InterchangeUsdTranslator::Private
 		pxr::UsdSkelSkinningQuery& PxrSkelSkinningQueryPtr = SkinningQuery;
 		pxr::UsdSkelSkeletonQuery& PxrSkelSkeletonQueryPtr = SkelQuery;
 
-		bool bSuccess = UsdToUnreal::ConvertGeomMesh(Prim, TempMeshDescription, TempMaterialInfo, Options, &PxrSkelSkinningQueryPtr, &PxrSkelSkeletonQueryPtr);
+		bool bSuccess = UsdToUnreal::ConvertGeomMesh(
+			Prim,
+			TempMeshDescription,
+			TempMaterialInfo,
+			Options,
+			&PxrSkelSkinningQueryPtr,
+			&PxrSkelSkeletonQueryPtr
+		);
 		if (!bSuccess)
 		{
 			return false;
@@ -3136,6 +3260,8 @@ UInterchangeUsdTranslatorSettings::UInterchangeUsdTranslatorSettings()
 		  0.01,				   // MetersPerUnit
 		  EUsdUpAxis::ZAxis	   // UpAxis
 	  }
+	, bTranslatePrimAttributes(false) // False by default as it could be expensive to traverse all attributes of all prims running the regex
+	, AttributeRegexFilter(TEXT("."))
 {
 }
 
diff --git a/Engine/Plugins/Interchange/Runtime/Source/Import/Public/Usd/InterchangeUsdTranslator.h b/Engine/Plugins/Interchange/Runtime/Source/Import/Public/Usd/InterchangeUsdTranslator.h
--- a/Engine/Plugins/Interchange/Runtime/Source/Import/Public/Usd/InterchangeUsdTranslator.h
+++ b/Engine/Plugins/Interchange/Runtime/Source/Import/Public/Usd/InterchangeUsdTranslator.h
@@ -52,6 +52,14 @@ public:
 	UPROPERTY(EditAnywhere, Category = "USD Translator", meta = (EditCondition = bOverrideStageOptions))
 	FUsdStageOptions StageOptions;
 
+	/** Whether to convert USD prim attributes into translated node custom attributes, if they pass the regex filter */
+	UPROPERTY(EditAnywhere, Category = "USD Translator")
+	bool bTranslatePrimAttributes;
+
+	/** Regex filter to select which USD prim attributes should be converted into translated node custom attributes */
+	UPROPERTY(EditAnywhere, Category = "USD Translator", meta = (EditCondition = bTranslatePrimAttributes))
+	FString AttributeRegexFilter;
+
 public:
 	UInterchangeUsdTranslatorSettings();
 };
diff --git a/Engine/Plugins/Runtime/USDCore/Source/UnrealUSDWrapper/Private/UsdWrappers/UsdAttribute.cpp b/Engine/Plugins/Runtime/USDCore/Source/UnrealUSDWrapper/Private/UsdWrappers/UsdAttribute.cpp
--- a/Engine/Plugins/Runtime/USDCore/Source/UnrealUSDWrapper/Private/UsdWrappers/UsdAttribute.cpp
+++ b/Engine/Plugins/Runtime/USDCore/Source/UnrealUSDWrapper/Private/UsdWrappers/UsdAttribute.cpp
@@ -215,6 +215,15 @@ namespace UE
 #endif	  // #if USE_USD_SDK
 	}
 
+	FString FUsdAttribute::GetCPPTypeName() const
+	{
+#if USE_USD_SDK
+		return UTF8_TO_TCHAR(Impl->PxrUsdAttribute.Get().GetTypeName().GetCPPTypeName().c_str());
+#else
+		return FString();
+#endif	  // #if USE_USD_SDK
+	}
+
 	bool FUsdAttribute::GetTimeSamples(TArray<double>& Times) const
 	{
 #if USE_USD_SDK
@@ -368,40 +377,415 @@ namespace UE
 #endif	  // #if USE_USD_SDK
 	}
 
-	template <typename T>
-	bool FUsdAttribute::Get(T& Value, TOptional<double> Time) const
-	{
 #if USE_USD_SDK
-		FScopedUsdAllocs UsdAllocs;
+	// Utility to help map from a UE type to an USD type within the Get<T> function
+	template<typename T>
+	struct USDTypeHelper
+	{
+	};
 
-		pxr::UsdTimeCode TimeCode = Time.IsSet() ? Time.GetValue() : pxr::UsdTimeCode::Default();
-		return Impl->PxrUsdAttribute.Get().Get(&Value, TimeCode);
-#else
-		return false;
-#endif	  // #if USE_USD_SDK
-	}
+	template<>
+	struct USDTypeHelper<float>
+	{
+		using USDType = float;
+	};
 
-	template UNREALUSDWRAPPER_API bool FUsdAttribute::Get(float& Value, TOptional<double> Time) const;
-	template UNREALUSDWRAPPER_API bool FUsdAttribute::Get(bool& Value, TOptional<double> Time) const;
+	template<>
+	struct USDTypeHelper<double>
+	{
+		using USDType = double;
+	};
+
+	template<>
+	struct USDTypeHelper<FFloat16>
+	{
+		using USDType = pxr::GfHalf;
+	};
+
+	template<>
+	struct USDTypeHelper<FQuat4f>
+	{
+		using USDType = pxr::GfQuatf;
+		using VecType = pxr::GfVec3f;
+	};
+
+	template<>
+	struct USDTypeHelper<FQuat4d>
+	{
+		using USDType = pxr::GfQuatd;
+		using VecType = pxr::GfVec3d;
+	};
+
+	template<>
+	struct USDTypeHelper<FVector2DHalf>
+	{
+		using USDType = pxr::GfVec2h;
+	};
+
+	template<>
+	struct USDTypeHelper<FVector2f>
+	{
+		using USDType = pxr::GfVec2f;
+	};
+
+	template<>
+	struct USDTypeHelper<FVector2d>
+	{
+		using USDType = pxr::GfVec2d;
+	};
+
+	template<>
+	struct USDTypeHelper<FIntPoint>
+	{
+		using USDType = pxr::GfVec2i;
+	};
+
+	template<>
+	struct USDTypeHelper<FVector3f>
+	{
+		using USDType = pxr::GfVec3f;
+	};
+
+	template<>
+	struct USDTypeHelper<FVector3d>
+	{
+		using USDType = pxr::GfVec3d;
+	};
+
+	template<>
+	struct USDTypeHelper<FIntVector>
+	{
+		using USDType = pxr::GfVec3i;
+	};
+
+	template<>
+	struct USDTypeHelper<FVector4f>
+	{
+		using USDType = pxr::GfVec4f;
+	};
+
+	template<>
+	struct USDTypeHelper<FVector4d>
+	{
+		using USDType = pxr::GfVec4d;
+	};
+
+	template<>
+	struct USDTypeHelper<FIntRect>
+	{
+		using USDType = pxr::GfVec4i;
+	};
+#endif	  // USE_USD_SDK
 
-	template <>
-	UNREALUSDWRAPPER_API bool FUsdAttribute::Get(FLinearColor& Value, TOptional<double> Time) const
+	template<typename T>
+	bool FUsdAttribute::Get(T& Value, TOptional<double> Time) const
 	{
 #if USE_USD_SDK
 		FScopedUsdAllocs UsdAllocs;
 
 		pxr::UsdTimeCode TimeCode = Time.IsSet() ? Time.GetValue() : pxr::UsdTimeCode::Default();
-		pxr::GfVec3f PxrValue;
-		bool bSuccess = Impl->PxrUsdAttribute.Get().Get(&PxrValue, TimeCode);
-		if (bSuccess)
+		pxr::UsdAttribute Attr = Impl->PxrUsdAttribute.Get();
+
+		if constexpr (std::is_floating_point_v<T> || std::is_same_v<T, FFloat16>)
+		{
+			using USDType = typename USDTypeHelper<T>::USDType;
+
+			pxr::VtValue VtValue;
+			if (Attr.Get(&VtValue, TimeCode))
+			{
+				// Cast all floating-point-like types to the desired floating point type
+				pxr::VtValue CastVtValue = VtValue.Cast<USDType>();
+				if (!CastVtValue.IsEmpty())
+				{
+					Value = static_cast<T>(CastVtValue.Get<USDType>());
+					return true;
+				}
+
+				if (VtValue.IsHolding<pxr::SdfTimeCode>())
+				{
+					const pxr::SdfTimeCode& TimeCodeValue = VtValue.Get<pxr::SdfTimeCode>();
+					Value = static_cast<T>(TimeCodeValue.GetValue());
+					return true;
+				}
+			}
+		}
+		else if constexpr (std::is_same_v<T, FString>)
 		{
-			Value = FLinearColor(PxrValue[0], PxrValue[1], PxrValue[2]);
+			pxr::TfToken TypeName = Attr.GetTypeName().GetAsToken();
+			if (TypeName == pxr::SdfValueTypeNames->String)
+			{
+				std::string PxrValue;
+				if (Attr.Get(&PxrValue, TimeCode))
+				{
+					Value = UTF8_TO_TCHAR(PxrValue.c_str());
+					return true;
+				}
+			}
+			else if (TypeName == pxr::SdfValueTypeNames->Token)
+			{
+				pxr::TfToken PxrValue;
+				if (Attr.Get(&PxrValue, TimeCode))
+				{
+					Value = UTF8_TO_TCHAR(PxrValue.GetString().c_str());
+					return true;
+				}
+			}
+			else if (TypeName == pxr::SdfValueTypeNames->Asset)
+			{
+				pxr::SdfAssetPath PxrValue;
+				if (Attr.Get(&PxrValue, TimeCode))
+				{
+					Value = UTF8_TO_TCHAR(PxrValue.GetAssetPath().c_str());
+					return true;
+				}
+			}
+			// As a convenience, let's use USD to stringify any attribute to a string, so we can do Get<FString> on anything
+			else
+			{
+				pxr::VtValue VtValue;
+				if (Attr.Get(&VtValue, TimeCode))
+				{
+					std::string UsdStr = pxr::TfStringify(VtValue);
+					Value = UTF8_TO_TCHAR(UsdStr.c_str());
+					return true;
+				}
+			}
+		}
+		else if constexpr (std::is_same_v<T, FMatrix44f> || std::is_same_v<T, FMatrix44d>)
+		{
+			pxr::TfToken TypeName = Attr.GetTypeName().GetAsToken();
+			if (TypeName == pxr::SdfValueTypeNames->Matrix2d)
+			{
+				pxr::GfMatrix2d Matrix;
+				if (Attr.Get(&Matrix, TimeCode))
+				{
+					Value = T(
+						Math::TPlane<typename T::FReal>(Matrix[0][0], Matrix[0][1], 0.0f, 0.0f),
+						Math::TPlane<typename T::FReal>(Matrix[1][0], Matrix[1][1], 0.0f, 0.0f),
+						Math::TPlane<typename T::FReal>(0.0f, 0.0f, 0.0f, 0.0f),
+						Math::TPlane<typename T::FReal>(0.0f, 0.0f, 0.0f, 0.0f)
+					);
+					return true;
+				}
+			}
+			else if (TypeName == pxr::SdfValueTypeNames->Matrix3d)
+			{
+				pxr::GfMatrix3d Matrix;
+				if (Attr.Get(&Matrix, TimeCode))
+				{
+					Value = T(
+						Math::TPlane<typename T::FReal>(Matrix[0][0], Matrix[0][1], Matrix[0][2], 0.0f),
+						Math::TPlane<typename T::FReal>(Matrix[1][0], Matrix[1][1], Matrix[1][2], 0.0f),
+						Math::TPlane<typename T::FReal>(Matrix[2][0], Matrix[2][1], Matrix[2][2], 0.0f),
+						Math::TPlane<typename T::FReal>(0.0f, 0.0f, 0.0f, 0.0f)
+					);
+					return true;
+				}
+			}
+			else if (TypeName == pxr::SdfValueTypeNames->Matrix4d || TypeName == pxr::SdfValueTypeNames->Frame4d)
+			{
+				pxr::GfMatrix4d Matrix;
+				if (Attr.Get(&Matrix, TimeCode))
+				{
+					Value = T(
+						Math::TPlane<typename T::FReal>(Matrix[0][0], Matrix[0][1], Matrix[0][2], Matrix[0][3]),
+						Math::TPlane<typename T::FReal>(Matrix[1][0], Matrix[1][1], Matrix[1][2], Matrix[1][3]),
+						Math::TPlane<typename T::FReal>(Matrix[2][0], Matrix[2][1], Matrix[2][2], Matrix[2][3]),
+						Math::TPlane<typename T::FReal>(Matrix[3][0], Matrix[3][1], Matrix[3][2], Matrix[3][3])
+					);
+					return true;
+				}
+			}
+		}
+		else if constexpr (std::is_same_v<T, FQuat4f> || std::is_same_v<T, FQuat4d>)
+		{
+			using USDType = typename USDTypeHelper<T>::USDType;
+			using VecType = typename USDTypeHelper<T>::VecType;
+
+			pxr::VtValue VtValue;
+			if (Attr.Get(&VtValue, TimeCode))
+			{
+				if (VtValue.CanCast<USDType>())
+				{
+					pxr::VtValue CastVtValue = VtValue.Cast<USDType>();
+					const USDType& CastValue = CastVtValue.Get<USDType>();
+					const VecType& Imaginary = CastValue.GetImaginary();
+					Value = T{Imaginary[0], Imaginary[1], Imaginary[2], CastValue.GetReal()};
+					return true;
+				}
+				// It doesn't seem like USD has conversions between the quaternion types, so we must check
+				// for Quath manually
+				else if (VtValue.CanCast<pxr::GfQuath>())
+				{
+					pxr::VtValue CastVtValue = VtValue.Cast<pxr::GfQuath>();
+					const pxr::GfQuath& CastValue = CastVtValue.Get<pxr::GfQuath>();
+					const pxr::GfVec3h& Imaginary = CastValue.GetImaginary();
+					Value = T{Imaginary[0], Imaginary[1], Imaginary[2], CastValue.GetReal()};
+					return true;
+				}
+			}
+		}
+		else if constexpr (std::is_same_v<T, FVector2DHalf> ||	  //
+						   std::is_same_v<T, FVector2f> ||		  //
+						   std::is_same_v<T, FVector2d> ||		  //
+						   std::is_same_v<T, FIntPoint>)
+		{
+			using USDType = typename USDTypeHelper<T>::USDType;
+
+			pxr::VtValue VtValue;
+			if (Attr.Get(&VtValue, TimeCode))
+			{
+				if (VtValue.CanCast<USDType>())
+				{
+					// This will handle all types with C++ type Vec2h/f/g/i (including roles like point, color, etc.)
+					pxr::VtValue CastVtValue = VtValue.Cast<USDType>();
+					const USDType& CastValue = CastVtValue.template Get<USDType>();
+					Value = T{CastValue[0], CastValue[1]};
+					return true;
+				}
+			}
+		}
+		else if constexpr (std::is_same_v<T, FVector3f> ||	  //
+						   std::is_same_v<T, FVector3d> ||	  //
+						   std::is_same_v<T, FIntVector>)
+		{
+			using USDType = typename USDTypeHelper<T>::USDType;
+
+			pxr::VtValue VtValue;
+			if (Attr.Get(&VtValue, TimeCode))
+			{
+				if (VtValue.CanCast<USDType>())
+				{
+					// This will handle all types with C++ type Vec3h/f/g/i (including roles like point, color, etc.)
+					pxr::VtValue CastVtValue = VtValue.Cast<USDType>();
+					const USDType& CastValue = CastVtValue.template Get<USDType>();
+					Value = T{CastValue[0], CastValue[1], CastValue[2]};
+					return true;
+				}
+			}
+		}
+		else if constexpr (std::is_same_v<T, FVector4f> ||	  //
+						   std::is_same_v<T, FVector4d> ||	  //
+						   std::is_same_v<T, FIntRect>)
+		{
+			using USDType = typename USDTypeHelper<T>::USDType;
+
+			pxr::VtValue VtValue;
+			if (Attr.Get(&VtValue, TimeCode))
+			{
+				if (VtValue.CanCast<USDType>())
+				{
+					// This will handle all types with C++ type Vec4h/f/g/i (including roles like point, color, etc.)
+					pxr::VtValue CastVtValue = VtValue.Cast<USDType>();
+					const USDType& CastValue = CastVtValue.template Get<USDType>();
+					Value = T{CastValue[0], CastValue[1], CastValue[2], CastValue[3]};
+					return true;
+				}
+			}
+		}
+		else if constexpr (std::is_same_v<T, FLinearColor> || std::is_same_v<T, FColor>)
+		{
+			pxr::VtValue VtValue;
+			if (Attr.Get(&VtValue, TimeCode))
+			{
+				// Vec3/Vec4 of floating types can have the color role in USD, so let's convert only those
+				if (VtValue.CanCast<pxr::GfVec4f>())
+				{
+					pxr::VtValue CastVtValue = VtValue.Cast<pxr::GfVec4f>();
+					const pxr::GfVec4f& CastValue = CastVtValue.Get<pxr::GfVec4f>();
+					if constexpr (std::is_same_v<T, FColor>)
+					{
+						// Color in USD is always energy linear, so perform sRGB conversion if going to FColor
+						const bool bSRGB = true;
+						Value = FLinearColor{CastValue[0], CastValue[1], CastValue[2], CastValue[3]}.ToFColor(bSRGB);
+					}
+					else
+					{
+						Value = FLinearColor{CastValue[0], CastValue[1], CastValue[2], CastValue[3]};
+					}
+					return true;
+				}
+				else if (VtValue.CanCast<pxr::GfVec3f>())
+				{
+					pxr::VtValue CastVtValue = VtValue.Cast<pxr::GfVec3f>();
+					const pxr::GfVec3f& CastValue = CastVtValue.Get<pxr::GfVec3f>();
+					if constexpr (std::is_same_v<T, FColor>)
+					{
+						const bool bSRGB = true;
+						Value = FLinearColor{CastValue[0], CastValue[1], CastValue[2], 1.0f}.ToFColor(bSRGB);
+					}
+					else
+					{
+						Value = FLinearColor{CastValue[0], CastValue[1], CastValue[2], 1.0f};
+					}
+					return true;
+				}
+			}
+		}
+		else if constexpr (std::is_same_v<T, TArray<uint8>>)
+		{
+			pxr::TfToken TypeName = Attr.GetTypeName().GetAsToken();
+			if (TypeName == pxr::SdfValueTypeNames->UCharArray)
+			{
+				pxr::VtArray<unsigned char> UsdValue;
+				if (Attr.Get(&UsdValue, TimeCode))
+				{
+					Value.SetNumUninitialized(UsdValue.size());
+					FMemory::Memcpy(Value.GetData(), UsdValue.cdata(), Value.Num() * Value.GetTypeSize());
+					return true;
+				}
+			}
+		}
+		else
+		{
+			return Attr.Get(&Value, TimeCode);
 		}
-		return bSuccess;
-#else
-		return false;
 #endif	  // #if USE_USD_SDK
+
+		return false;
 	}
+
+	template UNREALUSDWRAPPER_API bool FUsdAttribute::Get(bool& Value, TOptional<double> Time) const;
+
+	template UNREALUSDWRAPPER_API bool FUsdAttribute::Get(uint8& Value, TOptional<double> Time) const;
+	template UNREALUSDWRAPPER_API bool FUsdAttribute::Get(int32& Value, TOptional<double> Time) const;
+	template UNREALUSDWRAPPER_API bool FUsdAttribute::Get(uint32& Value, TOptional<double> Time) const;
+
+	// Note: int64 and uint64 are typedefs of "long long" and "unsigned long long", which I think on Linux can be *more* than
+	// 64 bits? USD doesn't define SdfValueTypeTraits<long long> though, but instead SdfValueTypeTraits<int64_t>, so
+	// when those differ we fail to explicitly instantiate these templates... So here we just use the same types specifically
+	template UNREALUSDWRAPPER_API bool FUsdAttribute::Get(int64_t& Value, TOptional<double> Time) const;
+	template UNREALUSDWRAPPER_API bool FUsdAttribute::Get(uint64_t& Value, TOptional<double> Time) const;
+
+	template UNREALUSDWRAPPER_API bool FUsdAttribute::Get(FFloat16& Value, TOptional<double> Time) const;
+	template UNREALUSDWRAPPER_API bool FUsdAttribute::Get(float& Value, TOptional<double> Time) const;
+	template UNREALUSDWRAPPER_API bool FUsdAttribute::Get(double& Value, TOptional<double> Time) const;
+
+	template UNREALUSDWRAPPER_API bool FUsdAttribute::Get(FString& Value, TOptional<double> Time) const;
+
+	template UNREALUSDWRAPPER_API bool FUsdAttribute::Get(FMatrix44f& Value, TOptional<double> Time) const;
+	template UNREALUSDWRAPPER_API bool FUsdAttribute::Get(FMatrix44d& Value, TOptional<double> Time) const;
+
+	template UNREALUSDWRAPPER_API bool FUsdAttribute::Get(FQuat4f& Value, TOptional<double> Time) const;
+	template UNREALUSDWRAPPER_API bool FUsdAttribute::Get(FQuat4d& Value, TOptional<double> Time) const;
+
+	template UNREALUSDWRAPPER_API bool FUsdAttribute::Get(FVector2DHalf& Value, TOptional<double> Time) const;
+	template UNREALUSDWRAPPER_API bool FUsdAttribute::Get(FVector2f& Value, TOptional<double> Time) const;
+	template UNREALUSDWRAPPER_API bool FUsdAttribute::Get(FVector2d& Value, TOptional<double> Time) const;
+	template UNREALUSDWRAPPER_API bool FUsdAttribute::Get(FIntPoint& Value, TOptional<double> Time) const;
+
+	template UNREALUSDWRAPPER_API bool FUsdAttribute::Get(FVector3f& Value, TOptional<double> Time) const;
+	template UNREALUSDWRAPPER_API bool FUsdAttribute::Get(FVector3d& Value, TOptional<double> Time) const;
+	template UNREALUSDWRAPPER_API bool FUsdAttribute::Get(FIntVector& Value, TOptional<double> Time) const;
+
+	template UNREALUSDWRAPPER_API bool FUsdAttribute::Get(FVector4f& Value, TOptional<double> Time) const;
+	template UNREALUSDWRAPPER_API bool FUsdAttribute::Get(FVector4d& Value, TOptional<double> Time) const;
+	template UNREALUSDWRAPPER_API bool FUsdAttribute::Get(FIntRect& Value, TOptional<double> Time) const;
+
+	template UNREALUSDWRAPPER_API bool FUsdAttribute::Get(FLinearColor& Value, TOptional<double> Time) const;
+	template UNREALUSDWRAPPER_API bool FUsdAttribute::Get(FColor& Value, TOptional<double> Time) const;
+
+	template UNREALUSDWRAPPER_API bool FUsdAttribute::Get(TArray<uint8>& Value, TOptional<double> Time) const;
 }	 // namespace UE
 
 namespace UsdUtils
@@ -411,7 +795,7 @@ namespace UsdUtils
 	{
 		if (!Prim)
 		{
-			return {};
+			return ValueType{};
 		}
 
 		UE::FUsdAttribute Attribute = Prim.GetAttribute(AttributeName.GetData());
@@ -425,7 +809,41 @@ namespace UsdUtils
 		return Value;
 	}
 
-	template UNREALUSDWRAPPER_API float GetAttributeValue(const UE::FUsdPrim& Prim, FStringView AttributeName, TOptional<double> Time);
 	template UNREALUSDWRAPPER_API bool GetAttributeValue(const UE::FUsdPrim& Prim, FStringView AttributeName, TOptional<double> Time);
+
+	template UNREALUSDWRAPPER_API uint8 GetAttributeValue(const UE::FUsdPrim& Prim, FStringView AttributeName, TOptional<double> Time);
+	template UNREALUSDWRAPPER_API int32 GetAttributeValue(const UE::FUsdPrim& Prim, FStringView AttributeName, TOptional<double> Time);
+	template UNREALUSDWRAPPER_API uint32 GetAttributeValue(const UE::FUsdPrim& Prim, FStringView AttributeName, TOptional<double> Time);
+	template UNREALUSDWRAPPER_API int64_t GetAttributeValue(const UE::FUsdPrim& Prim, FStringView AttributeName, TOptional<double> Time);
+	template UNREALUSDWRAPPER_API uint64_t GetAttributeValue(const UE::FUsdPrim& Prim, FStringView AttributeName, TOptional<double> Time);
+
+	template UNREALUSDWRAPPER_API FFloat16 GetAttributeValue(const UE::FUsdPrim& Prim, FStringView AttributeName, TOptional<double> Time);
+	template UNREALUSDWRAPPER_API float GetAttributeValue(const UE::FUsdPrim& Prim, FStringView AttributeName, TOptional<double> Time);
+	template UNREALUSDWRAPPER_API double GetAttributeValue(const UE::FUsdPrim& Prim, FStringView AttributeName, TOptional<double> Time);
+
+	template UNREALUSDWRAPPER_API FString GetAttributeValue(const UE::FUsdPrim& Prim, FStringView AttributeName, TOptional<double> Time);
+
+	template UNREALUSDWRAPPER_API FMatrix44f GetAttributeValue(const UE::FUsdPrim& Prim, FStringView AttributeName, TOptional<double> Time);
+	template UNREALUSDWRAPPER_API FMatrix44d GetAttributeValue(const UE::FUsdPrim& Prim, FStringView AttributeName, TOptional<double> Time);
+
+	template UNREALUSDWRAPPER_API FQuat4f GetAttributeValue(const UE::FUsdPrim& Prim, FStringView AttributeName, TOptional<double> Time);
+	template UNREALUSDWRAPPER_API FQuat4d GetAttributeValue(const UE::FUsdPrim& Prim, FStringView AttributeName, TOptional<double> Time);
+
+	template UNREALUSDWRAPPER_API FVector2DHalf GetAttributeValue(const UE::FUsdPrim& Prim, FStringView AttributeName, TOptional<double> Time);
+	template UNREALUSDWRAPPER_API FVector2f GetAttributeValue(const UE::FUsdPrim& Prim, FStringView AttributeName, TOptional<double> Time);
+	template UNREALUSDWRAPPER_API FVector2d GetAttributeValue(const UE::FUsdPrim& Prim, FStringView AttributeName, TOptional<double> Time);
+	template UNREALUSDWRAPPER_API FIntPoint GetAttributeValue(const UE::FUsdPrim& Prim, FStringView AttributeName, TOptional<double> Time);
+
+	template UNREALUSDWRAPPER_API FVector3f GetAttributeValue(const UE::FUsdPrim& Prim, FStringView AttributeName, TOptional<double> Time);
+	template UNREALUSDWRAPPER_API FVector3d GetAttributeValue(const UE::FUsdPrim& Prim, FStringView AttributeName, TOptional<double> Time);
+	template UNREALUSDWRAPPER_API FIntVector GetAttributeValue(const UE::FUsdPrim& Prim, FStringView AttributeName, TOptional<double> Time);
+
+	template UNREALUSDWRAPPER_API FVector4f GetAttributeValue(const UE::FUsdPrim& Prim, FStringView AttributeName, TOptional<double> Time);
+	template UNREALUSDWRAPPER_API FVector4d GetAttributeValue(const UE::FUsdPrim& Prim, FStringView AttributeName, TOptional<double> Time);
+	template UNREALUSDWRAPPER_API FIntRect GetAttributeValue(const UE::FUsdPrim& Prim, FStringView AttributeName, TOptional<double> Time);
+
 	template UNREALUSDWRAPPER_API FLinearColor GetAttributeValue(const UE::FUsdPrim& Prim, FStringView AttributeName, TOptional<double> Time);
-}
\ No newline at end of file
+	template UNREALUSDWRAPPER_API FColor GetAttributeValue(const UE::FUsdPrim& Prim, FStringView AttributeName, TOptional<double> Time);
+
+	template UNREALUSDWRAPPER_API TArray<uint8> GetAttributeValue(const UE::FUsdPrim& Prim, FStringView AttributeName, TOptional<double> Time);
+}	 // namespace UsdUtils
diff --git a/Engine/Plugins/Runtime/USDCore/Source/UnrealUSDWrapper/Public/UsdWrappers/UsdAttribute.h b/Engine/Plugins/Runtime/USDCore/Source/UnrealUSDWrapper/Public/UsdWrappers/UsdAttribute.h
--- a/Engine/Plugins/Runtime/USDCore/Source/UnrealUSDWrapper/Public/UsdWrappers/UsdAttribute.h
+++ b/Engine/Plugins/Runtime/USDCore/Source/UnrealUSDWrapper/Public/UsdWrappers/UsdAttribute.h
@@ -78,6 +78,7 @@ namespace UE
 		FName GetName() const;
 		FName GetBaseName() const;
 		FName GetTypeName() const;
+		FString GetCPPTypeName() const;
 
 		bool GetTimeSamples(TArray<double>& Times) const;
 		size_t GetNumTimeSamples() const;
