USDCore: simplify comparison operators for layer and stage wrappers

With C++20, the equality operator is now considered a "primary"
operator and the inequality operator is "secondary", where the latter
can be rewritten in terms of the former. As a result, we do not need to
define our own inequality operators.

The placeholder pointer types when the USD SDK is not available were also
modified to derive from a common base class. This avoids ambiguity around
which equality operator to use when comparing the two types, since they
can each be conversion constructed from the other.

This allows removing the setting to build the UnrealUSDWrapper
module using the C++17 standard, since it will now build successfully
using the default C++20 standard.

diff --git a/Engine/Plugins/Runtime/USDCore/Source/UnrealUSDWrapper/Private/UsdWrappers/SdfLayer.cpp b/Engine/Plugins/Runtime/USDCore/Source/UnrealUSDWrapper/Private/UsdWrappers/SdfLayer.cpp
--- a/Engine/Plugins/Runtime/USDCore/Source/UnrealUSDWrapper/Private/UsdWrappers/SdfLayer.cpp
+++ b/Engine/Plugins/Runtime/USDCore/Source/UnrealUSDWrapper/Private/UsdWrappers/SdfLayer.cpp
@@ -135,18 +135,6 @@ namespace UE
 	template bool UNREALUSDWRAPPER_API FSdfLayerWeak::operator==(const FSdfLayer& Other) const;
 	template bool UNREALUSDWRAPPER_API FSdfLayerWeak::operator==(const FSdfLayerWeak& Other) const;
 
-	template<typename PtrType>
-	template<typename OtherPtrType>
-	bool FSdfLayerBase<PtrType>::operator!=(const FSdfLayerBase<OtherPtrType>& Other) const
-	{
-		return !(*this == Other);
-	}
-
-	template bool UNREALUSDWRAPPER_API FSdfLayer::operator!=(const FSdfLayer& Other) const;
-	template bool UNREALUSDWRAPPER_API FSdfLayer::operator!=(const FSdfLayerWeak& Other) const;
-	template bool UNREALUSDWRAPPER_API FSdfLayerWeak::operator!=(const FSdfLayer& Other) const;
-	template bool UNREALUSDWRAPPER_API FSdfLayerWeak::operator!=(const FSdfLayerWeak& Other) const;
-
 	template<typename PtrType>
 	FSdfLayerBase<PtrType>::operator bool() const
 	{
diff --git a/Engine/Plugins/Runtime/USDCore/Source/UnrealUSDWrapper/Private/UsdWrappers/UsdStage.cpp b/Engine/Plugins/Runtime/USDCore/Source/UnrealUSDWrapper/Private/UsdWrappers/UsdStage.cpp
--- a/Engine/Plugins/Runtime/USDCore/Source/UnrealUSDWrapper/Private/UsdWrappers/UsdStage.cpp
+++ b/Engine/Plugins/Runtime/USDCore/Source/UnrealUSDWrapper/Private/UsdWrappers/UsdStage.cpp
@@ -132,18 +132,6 @@ namespace UE
 	template bool UNREALUSDWRAPPER_API FUsdStageWeak::operator==(const FUsdStage& Other) const;
 	template bool UNREALUSDWRAPPER_API FUsdStageWeak::operator==(const FUsdStageWeak& Other) const;
 
-	template<typename PtrType>
-	template<typename OtherPtrType>
-	bool FUsdStageBase<PtrType>::operator!=(const FUsdStageBase<OtherPtrType>& Other) const
-	{
-		return !(*this == Other);
-	}
-
-	template bool UNREALUSDWRAPPER_API FUsdStage::operator!=(const FUsdStage& Other) const;
-	template bool UNREALUSDWRAPPER_API FUsdStage::operator!=(const FUsdStageWeak& Other) const;
-	template bool UNREALUSDWRAPPER_API FUsdStageWeak::operator!=(const FUsdStage& Other) const;
-	template bool UNREALUSDWRAPPER_API FUsdStageWeak::operator!=(const FUsdStageWeak& Other) const;
-
 	template<typename PtrType>
 	FUsdStageBase<PtrType>::~FUsdStageBase()
 	{
diff --git a/Engine/Plugins/Runtime/USDCore/Source/UnrealUSDWrapper/Public/UsdWrappers/ForwardDeclarations.h b/Engine/Plugins/Runtime/USDCore/Source/UnrealUSDWrapper/Public/UsdWrappers/ForwardDeclarations.h
--- a/Engine/Plugins/Runtime/USDCore/Source/UnrealUSDWrapper/Public/UsdWrappers/ForwardDeclarations.h
+++ b/Engine/Plugins/Runtime/USDCore/Source/UnrealUSDWrapper/Public/UsdWrappers/ForwardDeclarations.h
@@ -43,52 +43,36 @@ namespace UE
 	// We use these types so that we can keep the same FUsdStage, etc. type aliases
 	// and not change the rest of the code. It's the bare minimum to be able to
 	// be tested for validity (and always fail).
-	// We use two to prevent multiple symbol definitions on the wrapper copy
-	// constructor/assignment operators.
-	class FDummyWeakPtrType;
-
-	class FDummyRefPtrType
+	class FDummyPtrBase
 	{
 	public:
-		FDummyRefPtrType(){};
-		FDummyRefPtrType(const FDummyWeakPtrType&){};
-
 		explicit operator bool() const
 		{
 			return false;
 		}
 
-		bool operator==(const FDummyRefPtrType&) const
+		bool operator==(const FDummyPtrBase&) const
 		{
 			return false;
 		}
+	};
 
-		bool operator!=(const FDummyRefPtrType&) const
-		{
-			return false;
-		}
+	// We use separate ref and weak types to prevent multiple symbol
+	// definitions on the wrapper copy constructor/assignment operators.
+	class FDummyWeakPtrType;
+
+	class FDummyRefPtrType : public FDummyPtrBase
+	{
+	public:
+		FDummyRefPtrType(){};
+		FDummyRefPtrType(const FDummyWeakPtrType&){};
 	};
 
-	class FDummyWeakPtrType
+	class FDummyWeakPtrType : public FDummyPtrBase
 	{
 	public:
 		FDummyWeakPtrType(){};
 		FDummyWeakPtrType(const FDummyRefPtrType&){};
-
-		explicit operator bool() const
-		{
-			return false;
-		}
-
-		bool operator==(const FDummyWeakPtrType&) const
-		{
-			return false;
-		}
-
-		bool operator!=(const FDummyWeakPtrType&) const
-		{
-			return false;
-		}
 	};
 
 	template<typename PtrType>
diff --git a/Engine/Plugins/Runtime/USDCore/Source/UnrealUSDWrapper/Public/UsdWrappers/SdfLayer.h b/Engine/Plugins/Runtime/USDCore/Source/UnrealUSDWrapper/Public/UsdWrappers/SdfLayer.h
--- a/Engine/Plugins/Runtime/USDCore/Source/UnrealUSDWrapper/Public/UsdWrappers/SdfLayer.h
+++ b/Engine/Plugins/Runtime/USDCore/Source/UnrealUSDWrapper/Public/UsdWrappers/SdfLayer.h
@@ -90,12 +90,10 @@ namespace UE
 		FSdfLayerBase& operator=(const FSdfLayerWeak& Other);
 		FSdfLayerBase& operator=(FSdfLayerWeak&& Other);
 
-		// We have to implement templated comparison operators, or else we get ambigous
+		// We have to implement a templated comparison operator, or else we get ambigous
 		// conversions on the Mac targets when comparing FSdfLayer and FSdfLayerWeak
 		template<typename OtherPtrType>
 		bool operator==(const FSdfLayerBase<OtherPtrType>& Other) const;
-		template<typename OtherPtrType>
-		bool operator!=(const FSdfLayerBase<OtherPtrType>& Other) const;
 
 		explicit operator bool() const;
 
diff --git a/Engine/Plugins/Runtime/USDCore/Source/UnrealUSDWrapper/Public/UsdWrappers/UsdStage.h b/Engine/Plugins/Runtime/USDCore/Source/UnrealUSDWrapper/Public/UsdWrappers/UsdStage.h
--- a/Engine/Plugins/Runtime/USDCore/Source/UnrealUSDWrapper/Public/UsdWrappers/UsdStage.h
+++ b/Engine/Plugins/Runtime/USDCore/Source/UnrealUSDWrapper/Public/UsdWrappers/UsdStage.h
@@ -43,12 +43,10 @@ namespace UE
 
 		explicit operator bool() const;
 
-		// We have to implement templated comparison operators, or else we get ambigous
+		// We have to implement a templated comparison operator, or else we get ambigous
 		// conversions on the Mac targets when comparing FUsdStage and FUsdStageWeak
 		template<typename OtherPtrType>
 		bool operator==(const FUsdStageBase<OtherPtrType>& Other) const;
-		template<typename OtherPtrType>
-		bool operator!=(const FUsdStageBase<OtherPtrType>& Other) const;
 
 		// Auto conversion from/to PtrType. We use concrete pointer types here
 		// because we should also be able to convert between them
diff --git a/Engine/Plugins/Runtime/USDCore/Source/UnrealUSDWrapper/UnrealUSDWrapper.Build.cs b/Engine/Plugins/Runtime/USDCore/Source/UnrealUSDWrapper/UnrealUSDWrapper.Build.cs
--- a/Engine/Plugins/Runtime/USDCore/Source/UnrealUSDWrapper/UnrealUSDWrapper.Build.cs
+++ b/Engine/Plugins/Runtime/USDCore/Source/UnrealUSDWrapper/UnrealUSDWrapper.Build.cs
@@ -10,11 +10,6 @@ namespace UnrealBuildTool.Rules
 	{
 		public UnrealUSDWrapper(ReadOnlyTargetRules Target) : base(Target)
 		{
-			// Does not compile with C++20:
-			// error C2666: 'UE::FDummyWeakPtrType::operator ==': overloaded functions have similar conversions
-			// error C2666: 'UE::FDummyRefPtrType::operator ==': overloaded functions have similar conversions
-			CppStandard = CppStandardVersion.Cpp17;
-
 			bUseRTTI = true;
 
 			PublicDependencyModuleNames.AddRange(
