USD: Proper fix for int64/uint64 handling on Linux.

Additional preflight with the exact target that failed UE-239778: 6792696d9828b967f5ae91b3

diff --git a/Engine/Plugins/Interchange/Runtime/Source/Import/Private/Usd/InterchangeUsdTranslator.cpp b/Engine/Plugins/Interchange/Runtime/Source/Import/Private/Usd/InterchangeUsdTranslator.cpp
--- a/Engine/Plugins/Interchange/Runtime/Source/Import/Private/Usd/InterchangeUsdTranslator.cpp
+++ b/Engine/Plugins/Interchange/Runtime/Source/Import/Private/Usd/InterchangeUsdTranslator.cpp
@@ -2433,8 +2433,8 @@ namespace UE::InterchangeUsdTranslator::Private
 				{TEXT("uchar[]"), 		TranslateAttribute<TArray<uint8>>},
 				{TEXT("int"), 			TranslateAttribute<int32>},
 				{TEXT("uint"), 			TranslateAttribute<uint32>},
-				// {TEXT("int64"), 		TranslateAttribute<int64_t>},  // Temporarily disabled as they don't work on Linux
-				// {TEXT("uint64"), 		TranslateAttribute<uint64_t>},
+				{TEXT("int64"), 		TranslateAttribute<int64>},
+				{TEXT("uint64"), 		TranslateAttribute<uint64>},
 
 				{TEXT("half"),	 		TranslateAttribute<FFloat16>},
 				{TEXT("float"), 		TranslateAttribute<float>},
diff --git a/Engine/Plugins/Runtime/USDCore/Source/UnrealUSDWrapper/Private/UsdWrappers/UsdAttribute.cpp b/Engine/Plugins/Runtime/USDCore/Source/UnrealUSDWrapper/Private/UsdWrappers/UsdAttribute.cpp
--- a/Engine/Plugins/Runtime/USDCore/Source/UnrealUSDWrapper/Private/UsdWrappers/UsdAttribute.cpp
+++ b/Engine/Plugins/Runtime/USDCore/Source/UnrealUSDWrapper/Private/UsdWrappers/UsdAttribute.cpp
@@ -736,6 +736,30 @@ namespace UE
 				}
 			}
 		}
+		else if constexpr (std::is_same_v<T, int64>)
+		{
+			// Note: The Unreal int64 and uint64 are typedefs of "long long" and "unsigned long long".
+			// On Windows / MSVC these end up with a type that matches int64_t and uint64_t. On Linux / Clang this is not the case.
+			//
+			// USD only has templates specialized for int64_t/uint64_t, so if we use our typedef'd "long long"s here we will get
+			// a compile error on Linux / Clang, so we have to go through the other type.
+			int64_t TypedValue;
+			if (Attr.Get(&TypedValue, TimeCode))
+			{
+				Value = static_cast<T>(TypedValue);
+				return true;
+			}
+		}
+		else if constexpr (std::is_same_v<T, uint64>)
+		{
+			// See comment on the int64 case right above this
+			uint64_t TypedValue;
+			if (Attr.Get(&TypedValue, TimeCode))
+			{
+				Value = static_cast<T>(TypedValue);
+				return true;
+			}
+		}
 		else
 		{
 			return Attr.Get(&Value, TimeCode);
@@ -751,11 +775,8 @@ namespace UE
 	template UNREALUSDWRAPPER_API bool FUsdAttribute::Get(int32& Value, TOptional<double> Time) const;
 	template UNREALUSDWRAPPER_API bool FUsdAttribute::Get(uint32& Value, TOptional<double> Time) const;
 
-	// Note: int64 and uint64 are typedefs of "long long" and "unsigned long long", which I think on Linux can be *more* than
-	// 64 bits? USD doesn't define SdfValueTypeTraits<long long> though, but instead SdfValueTypeTraits<int64_t>, so
-	// when those differ we fail to explicitly instantiate these templates... So here we just use the same types specifically
-	template UNREALUSDWRAPPER_API bool FUsdAttribute::Get(int64_t& Value, TOptional<double> Time) const;
-	template UNREALUSDWRAPPER_API bool FUsdAttribute::Get(uint64_t& Value, TOptional<double> Time) const;
+	template UNREALUSDWRAPPER_API bool FUsdAttribute::Get(int64& Value, TOptional<double> Time) const;
+	template UNREALUSDWRAPPER_API bool FUsdAttribute::Get(uint64& Value, TOptional<double> Time) const;
 
 	template UNREALUSDWRAPPER_API bool FUsdAttribute::Get(FFloat16& Value, TOptional<double> Time) const;
 	template UNREALUSDWRAPPER_API bool FUsdAttribute::Get(float& Value, TOptional<double> Time) const;
@@ -814,8 +835,8 @@ namespace UsdUtils
 	template UNREALUSDWRAPPER_API uint8 GetAttributeValue(const UE::FUsdPrim& Prim, FStringView AttributeName, TOptional<double> Time);
 	template UNREALUSDWRAPPER_API int32 GetAttributeValue(const UE::FUsdPrim& Prim, FStringView AttributeName, TOptional<double> Time);
 	template UNREALUSDWRAPPER_API uint32 GetAttributeValue(const UE::FUsdPrim& Prim, FStringView AttributeName, TOptional<double> Time);
-	template UNREALUSDWRAPPER_API int64_t GetAttributeValue(const UE::FUsdPrim& Prim, FStringView AttributeName, TOptional<double> Time);
-	template UNREALUSDWRAPPER_API uint64_t GetAttributeValue(const UE::FUsdPrim& Prim, FStringView AttributeName, TOptional<double> Time);
+	template UNREALUSDWRAPPER_API int64 GetAttributeValue(const UE::FUsdPrim& Prim, FStringView AttributeName, TOptional<double> Time);
+	template UNREALUSDWRAPPER_API uint64 GetAttributeValue(const UE::FUsdPrim& Prim, FStringView AttributeName, TOptional<double> Time);
 
 	template UNREALUSDWRAPPER_API FFloat16 GetAttributeValue(const UE::FUsdPrim& Prim, FStringView AttributeName, TOptional<double> Time);
 	template UNREALUSDWRAPPER_API float GetAttributeValue(const UE::FUsdPrim& Prim, FStringView AttributeName, TOptional<double> Time);
