USD
- Refactored the code that creates the MaterialX shader graph

Interchange USD:
- Store the converted <geompropvalue> nodes into <image> into the Mesh Node
- Store the different primvar names in the MeshFactoryNode

diff --git a/Engine/Plugins/Interchange/Runtime/Source/Common/Public/Usd/InterchangeUsdDefinitions.h b/Engine/Plugins/Interchange/Runtime/Source/Common/Public/Usd/InterchangeUsdDefinitions.h
--- a/Engine/Plugins/Interchange/Runtime/Source/Common/Public/Usd/InterchangeUsdDefinitions.h
+++ b/Engine/Plugins/Interchange/Runtime/Source/Common/Public/Usd/InterchangeUsdDefinitions.h
@@ -31,4 +31,10 @@ namespace UE::Interchange::USD
 	// Note: To indicate TwoSided/Translucent we use the CustomTwoSided/BlendMode standard UInterchangeShaderGraphNode
 	// attributes.
 	const FString NeedsVTParentIdentifier = TEXT("USD_MI_NeedsVTParent");
+
+	namespace Primvar
+	{
+		const FString Number = TEXT("USD_PrimvarNumber");
+		const FString Name = TEXT("USD_PrimvarName"); // In case of "Number" primvars, the user should concatenate the Index to this attribute
+	}
 }
diff --git a/Engine/Plugins/Interchange/Runtime/Source/Import/Private/Usd/InterchangeUsdTranslator.cpp b/Engine/Plugins/Interchange/Runtime/Source/Import/Private/Usd/InterchangeUsdTranslator.cpp
--- a/Engine/Plugins/Interchange/Runtime/Source/Import/Private/Usd/InterchangeUsdTranslator.cpp
+++ b/Engine/Plugins/Interchange/Runtime/Source/Import/Private/Usd/InterchangeUsdTranslator.cpp
@@ -12,6 +12,7 @@
 #include "USDLayerUtils.h"
 #include "USDLightConversion.h"
 #include "USDLog.h"
+#include "USDMaterialXShaderGraph.h"
 #include "USDObjectUtils.h"
 #include "USDPrimConversion.h"
 #include "USDShadeConversion.h"
@@ -21,7 +22,6 @@
 #include "UsdWrappers/SdfPath.h"
 #include "UsdWrappers/UsdAttribute.h"
 #include "UsdWrappers/UsdEditContext.h"
-#include "UsdWrappers/UsdGeomXformable.h"
 #include "UsdWrappers/UsdPrim.h"
 #include "UsdWrappers/UsdRelationship.h"
 #include "UsdWrappers/UsdSkelAnimQuery.h"
@@ -41,7 +41,6 @@
 #include "InterchangeCameraNode.h"
 #include "InterchangeLightNode.h"
 #include "InterchangeManager.h"
-#include "InterchangeMaterialInstanceNode.h"
 #include "InterchangeMeshNode.h"
 #include "InterchangeSceneNode.h"
 #include "InterchangeShaderGraphNode.h"
@@ -283,11 +282,13 @@ namespace UE::InterchangeUsdTranslator::Private
 
 		/**
 		 * Search for a MaterialX file embedded in the file, in case there's none, create a ShaderGraph using MaterialX
+		 * @param GeomPropValueNames - Array of the <geompropvalue> node names from the shader graph that has been converted to <image> nodes (used for baking later on by the Factory)
 		 */
 		bool AddMaterialXShaderGraph(
-			const UE::FUsdPrim& Prim,
+			const UE::FUsdPrim& InPrim,
 			UInterchangeUsdTranslatorSettings* TranslatorSettings,
-			UInterchangeBaseNodeContainer& NodeContainer
+			UInterchangeBaseNodeContainer& NodeContainer,
+			TArray<FUsdMaterialXShaderGraph::FGeomProp>& OutGeomPropValues
 		);
 
 		/**
@@ -298,6 +299,7 @@ namespace UE::InterchangeUsdTranslator::Private
 			const FString& MaterialPrimName,
 			const FString& MaterialUid,
 			UInterchangeBaseNodeContainer& NodeContainer,
+			const TArray<FUsdMaterialXShaderGraph::FGeomProp>& GeomPropValues = {},
 			bool bIterateShaderGraphs = false
 		);
 
@@ -822,9 +824,11 @@ namespace UE::InterchangeUsdTranslator::Private
 		const FString& MaterialPrimName,
 		const FString& MaterialUid,
 		UInterchangeBaseNodeContainer& NodeContainer,
+		const TArray<FUsdMaterialXShaderGraph::FGeomProp>& GeomPropValues,
 		bool bIterateShaderGraphs
 	)
 	{
+		using namespace UE::Interchange;
 		if (bIterateShaderGraphs)
 		{
 			// The material from the MaterialXTranslator doesn't have the same UID, both the PrimPath have the same name but not the same path
@@ -853,6 +857,16 @@ namespace UE::InterchangeUsdTranslator::Private
 			{
 				for (const FMaterialSlotMesh& MaterialSlotMesh : *SlotMeshes)
 				{
+					if (!GeomPropValues.IsEmpty())
+					{
+						MaterialSlotMesh.MeshNode->AddInt32Attribute(USD::Primvar::Number, GeomPropValues.Num());
+
+						for (int32 Index = 0; Index < GeomPropValues.Num(); ++Index)
+						{
+							MaterialSlotMesh.MeshNode->AddStringAttribute(USD::Primvar::Name + FString::FromInt(Index), GeomPropValues[Index].Name);
+						}
+					}
+
 					if (!MaterialSlotMesh.MeshNode->GetSlotMaterialDependencyUid(MaterialSlotMesh.MaterialSlotName, *NewMaterialUID))
 					{
 						MaterialSlotMesh.MeshNode->SetSlotMaterialDependencyUid(MaterialSlotMesh.MaterialSlotName, *NewMaterialUID);
@@ -865,7 +879,8 @@ namespace UE::InterchangeUsdTranslator::Private
 	bool UInterchangeUSDTranslatorImpl::AddMaterialXShaderGraph(
 		const UE::FUsdPrim& Prim,
 		UInterchangeUsdTranslatorSettings* TranslatorSettings,
-		UInterchangeBaseNodeContainer& NodeContainer
+		UInterchangeBaseNodeContainer& NodeContainer,
+		TArray<FUsdMaterialXShaderGraph::FGeomProp>& GeomPropValues
 	)
 	{
 		FName RenderContext = TranslatorSettings ? TranslatorSettings->RenderContext : UnrealIdentifiers::UniversalRenderContext;
@@ -901,7 +916,9 @@ namespace UE::InterchangeUsdTranslator::Private
 			// Only enable the shader graph on Windows for the time being as it crashes on Linux cause of a probable double free
 #if PLATFORM_WINDOWS
 			bool bHasShaderGraph = false;
-			if (MaterialX::DocumentPtr Document = UsdToUnreal::CreateSurfaceMaterial(Prim, TranslatorSettings ? *RenderContext.ToString() : nullptr))
+			FUsdMaterialXShaderGraph ShaderGraph(Prim, *RenderContext.ToString());
+			GeomPropValues = ShaderGraph.GetGeomPropValueNames();
+			if (MaterialX::DocumentPtr Document = ShaderGraph.GetDocument())
 			{
 				if (MaterialX::NodePtr SurfaceMaterial = Document->getNode(TCHAR_TO_UTF8(*Prim.GetName().ToString())))
 				{
@@ -1102,9 +1119,9 @@ namespace UE::InterchangeUsdTranslator::Private
 		}
 
 		// we only create material instances if we didn't find any MaterialX instances (including a shader graph)
-		if (AddMaterialXShaderGraph(Prim, TranslatorSettings, NodeContainer))
+		if (TArray<FUsdMaterialXShaderGraph::FGeomProp> GeomPropValues; AddMaterialXShaderGraph(Prim, TranslatorSettings, NodeContainer, GeomPropValues))
 		{
-			SetMaterialSlotDependencies(MaterialPrimName, MaterialUid, NodeContainer, true);
+			SetMaterialSlotDependencies(MaterialPrimName, MaterialUid, NodeContainer, GeomPropValues, true);
 			return;
 		}
 
diff --git a/Engine/Plugins/Interchange/Runtime/Source/Pipelines/Private/InterchangeUSDPipeline.cpp b/Engine/Plugins/Interchange/Runtime/Source/Pipelines/Private/InterchangeUSDPipeline.cpp
--- a/Engine/Plugins/Interchange/Runtime/Source/Pipelines/Private/InterchangeUSDPipeline.cpp
+++ b/Engine/Plugins/Interchange/Runtime/Source/Pipelines/Private/InterchangeUSDPipeline.cpp
@@ -4,7 +4,8 @@
 
 #include "InterchangeManager.h"
 #include "InterchangeMaterialFactoryNode.h"
-#include "InterchangeMaterialInstanceNode.h"
+#include "InterchangeMeshNode.h"
+#include "InterchangeMeshFactoryNode.h"
 #include "InterchangeSceneNode.h"
 #include "InterchangeShaderGraphNode.h"
 #include "InterchangeSkeletalMeshLodDataNode.h"
@@ -427,6 +428,40 @@ namespace UE::InterchangeUsdPipeline::Private
 			NodeContainer.ReplaceNode(FactoryNodeUid, nullptr);
 		}
 	}
+
+	/**
+	 * Add the primvar attributes from the mesh nodes to factory nodes
+	 */
+	void AddPrimvarAttributesToFactoryNodes(UInterchangeBaseNodeContainer* NodeContainer)
+	{
+		using namespace UE::Interchange;
+
+		NodeContainer->IterateNodesOfType<UInterchangeMeshNode>(
+			[NodeContainer](const FString& NodeUid, UInterchangeMeshNode* MeshNode)
+			{
+				if (FString FactoryNodeUid = UInterchangeFactoryBaseNode::BuildFactoryNodeUid(MeshNode->GetUniqueID());
+					UInterchangeMeshFactoryNode * MeshFactoryNode = Cast<UInterchangeMeshFactoryNode>(NodeContainer->GetFactoryNode(FactoryNodeUid)))
+				{
+					if (int32 NumPrimvar; MeshNode->GetInt32Attribute(USD::Primvar::Number, NumPrimvar))
+					{
+						MeshFactoryNode->AddInt32Attribute(USD::Primvar::Number, NumPrimvar);
+						for (int32 Index = 0; Index < NumPrimvar; ++Index)
+						{
+							FString PrimvarAttribute = USD::Primvar::Name + FString::FromInt(Index);
+							if (FString PrimvarName; MeshNode->GetStringAttribute(PrimvarAttribute, PrimvarName))
+							{
+								MeshFactoryNode->AddStringAttribute(PrimvarAttribute, PrimvarName);
+							}
+
+							if (bool bTangentSpace; MeshNode->GetBooleanAttribute(PrimvarAttribute, bTangentSpace))
+							{
+								MeshFactoryNode->AddBooleanAttribute(PrimvarAttribute, bTangentSpace);
+							}
+						}
+					}
+				}
+			});
+	}
 }	 // namespace UE::InterchangeUsdPipeline::Private
 
 UInterchangeUsdPipeline::UInterchangeUsdPipeline()
@@ -458,6 +493,8 @@ void UInterchangeUsdPipeline::ExecutePipeline(
 	{
 		RemovePseudoRootFactoryNode(*NodeContainer);
 	}
+
+	AddPrimvarAttributesToFactoryNodes(NodeContainer);
 }
 
 bool UInterchangeUsdPipeline::CanExecuteOnAnyThread(EInterchangePipelineTask PipelineTask)
diff --git a/Engine/Plugins/Runtime/USDCore/Source/USDUtilities/Private/USDMaterialXShaderGraph.cpp b/Engine/Plugins/Runtime/USDCore/Source/USDUtilities/Private/USDMaterialXShaderGraph.cpp
new file mode 100644
--- /dev/null
+++ b/Engine/Plugins/Runtime/USDCore/Source/USDUtilities/Private/USDMaterialXShaderGraph.cpp
@@ -0,0 +1,522 @@
+// Copyright Epic Games, Inc. All Rights Reserved.
+
+#include "USDMaterialXShaderGraph.h"
+
+#ifdef ENABLE_USD_MATERIALX
+
+#include "USDLayerUtils.h"
+#include "USDLog.h"
+#include "USDMemory.h"
+#include "USDTypesConversion.h"
+#include "UsdWrappers/SdfLayer.h"
+#include "Misc/Paths.h"
+
+namespace mx = MaterialX;
+
+FUsdMaterialXShaderGraph::FUsdMaterialXShaderGraph(const pxr::UsdPrim& UsdShadeMaterialPrim, const TCHAR* RenderContext)
+	: MaterialXTypes {
+		"boolean",
+		"integer",
+		"float",
+		"color3",
+		"color4",
+		"vector2",
+		"vector3",
+		"vector4",
+		"matrix33",
+		"matrix44",
+		"string",
+		"filename",
+		"geomname",
+		"surfaceshader",
+		"displacementshader",
+		"volumeshader",
+		"lightshader",
+		"material",
+		"none",
+		"integerarray",
+		"floatarray",
+		"color3array",
+		"color4array",
+		"vector2array",
+		"vector3array",
+		"vector4array",
+		"stringarray",
+		"geomnamearray",
+	}
+	, UsdToMaterialXTypes {
+		{"color3f",  "color3"},
+		{"color4f",  "color4"},
+		{"float2",   "vector2"},
+		{"vector3f", "vector3"},
+		{"vector4f", "vector4"},
+		{"int",      "integer"},
+	}
+	, TangentSpaceInputs{
+	    //UsdPreviewSurface
+		"normal",
+		//Standard Surface
+		"coat_normal",
+		"tangent",
+		//Openpbr Surface
+	    "geometry_normal",
+		"geometry_coat_normal",
+		"geometry_tangent",
+		"geometry_coat_tangent"
+	}
+{
+	pxr::UsdShadeMaterial UsdShadeMaterial{ UsdShadeMaterialPrim };
+
+	if (!UsdShadeMaterial)
+	{
+		UE_LOG(LogUsd, Error, TEXT("Couldn't create a UsdShadeMaterial for the prim."));
+		return;
+	}
+
+	Document = mx::createDocument();
+	mx::DocumentPtr MaterialXLibrary = mx::createDocument();
+	mx::FileSearchPath MaterialXFolder{
+		TCHAR_TO_UTF8(*FPaths::Combine(FPaths::EngineDir(), TEXT("Binaries"), TEXT("ThirdParty"), TEXT("MaterialX")))
+	};
+	mx::StringSet LoadedLibs = mx::loadLibraries({ "libraries" }, MaterialXFolder, MaterialXLibrary);
+	Document->importLibrary(MaterialXLibrary);
+
+	CreateSurfaceMaterial(UsdShadeMaterial, RenderContext);
+}
+
+mx::DocumentPtr FUsdMaterialXShaderGraph::GetDocument() const
+{
+	return Document;
+}
+
+const TArray<FUsdMaterialXShaderGraph::FGeomProp>& FUsdMaterialXShaderGraph::GetGeomPropValueNames() const
+{
+	return GeomPropValueNames;
+}
+
+void FUsdMaterialXShaderGraph::CreateSurfaceMaterial(const pxr::UsdShadeMaterial& UsdShadeMaterial, const TCHAR* RenderContext)
+{
+	pxr::TfToken RenderContextToken = pxr::UsdShadeTokens->universalRenderContext;
+	if (RenderContext)
+	{
+		pxr::TfToken ProvidedRenderContextToken = UnrealToUsd::ConvertToken(RenderContext).Get();
+		RenderContextToken = ProvidedRenderContextToken;
+	}
+
+	pxr::UsdShadeShader SurfaceShader = UsdShadeMaterial.ComputeSurfaceSource(RenderContextToken);
+	pxr::UsdShadeShader DisplacementShader = UsdShadeMaterial.ComputeDisplacementSource(RenderContextToken);
+
+	if (!SurfaceShader)
+	{
+		UE_LOG(LogUsd, Error, TEXT("Couldn't find a SurfaceShader for UsdShadeMaterial '%s'."), UTF8_TO_TCHAR(UsdShadeMaterial.GetPrim().GetName().GetString().c_str()));
+		return;
+	}
+
+	auto CreateShaderGraph = [this](const pxr::UsdShadeShader& SurfaceShader, mx::DocumentPtr Document) -> mx::NodePtr
+	{
+		pxr::TfToken NodeDefTokenId;
+		SurfaceShader.GetShaderId(&NodeDefTokenId);
+		mx::NodeDefPtr NodeDef = Document->getNodeDef(NodeDefTokenId.GetString());
+
+		// no node definition, just return
+		if (!NodeDef)
+		{
+			return nullptr;
+		}
+
+		mx::NodePtr MtlxSurfaceShader = Document->addNodeInstance(NodeDef, SurfaceShader.GetPrim().GetName().GetString());
+
+		// if the shadergraph is not valid (probably caused by a missing nodedef or wrong type), just invalidate it all
+		if (!ComputeShaderGraph(MtlxSurfaceShader, SurfaceShader))
+		{
+			MtlxSurfaceShader = nullptr;
+		}
+
+		return MtlxSurfaceShader;
+	};
+
+	// Create and store the input interface names
+	TUsdStore<std::vector<pxr::UsdShadeInput>> InterfaceInputs = UsdShadeMaterial.GetInterfaceInputs();
+
+	for (const pxr::UsdShadeInput& InterfaceInput : *InterfaceInputs)
+	{
+		std::string InterfaceNameInput = InterfaceInput.GetBaseName().GetString();
+		std::string InterfaceNameType = InterfaceInput.GetTypeName().GetAsToken().GetString();
+		GetMatchingInputType(InterfaceNameInput, InterfaceNameType);
+
+		mx::InputPtr InputInterfaceName = Document->addInput(InterfaceNameInput, InterfaceNameType);
+		SetInputValue(InputInterfaceName, InterfaceInput);
+		InterfaceNames.insert(InputInterfaceName->getName());
+	}
+
+	// Create the shader graph for the surface shader and displacement shader
+	if (mx::NodePtr MaterialXSurfaceShader = CreateShaderGraph(SurfaceShader, Document))
+	{
+		mx::NodeDefPtr SurfaceMaterialNodeDef = Document->getNodeDef("ND_surfacematerial");
+		mx::NodePtr SurfaceMaterial = Document->addNodeInstance(SurfaceMaterialNodeDef, UsdShadeMaterial.GetPrim().GetName().GetString());
+
+		SurfaceMaterial->setConnectedNode("surfaceshader", MaterialXSurfaceShader);
+
+		// the displacement shader is not necessarily required
+		if (DisplacementShader)
+		{
+			if (mx::NodePtr MaterialXDisplacementShader = CreateShaderGraph(DisplacementShader, Document))
+			{
+				SurfaceMaterial->setConnectedNode("displacementshader", MaterialXDisplacementShader);
+			}
+		}
+
+		std::string ErrorMessage;
+		SurfaceMaterial->validate(&ErrorMessage);
+		if (!ErrorMessage.empty())
+		{
+			UE_LOG(LogUsd, Error, TEXT("USD MaterialX: %s."), UTF8_TO_TCHAR(ErrorMessage.c_str()));
+			SurfaceMaterial = nullptr;
+		}
+	}
+}
+
+bool FUsdMaterialXShaderGraph::ComputeShaderGraph(mx::NodePtr Node, const pxr::UsdShadeShader& Shader)
+{
+	if (!Node)
+	{
+		return false;
+	}
+
+	TUsdStore<std::vector<pxr::UsdShadeInput>> UsdInputs = Shader.GetInputs();
+
+	for (const pxr::UsdShadeInput& UsdInput : *UsdInputs)
+	{
+		// Basically the idea here is to traverse the entire graph coming from a tangent space input of a surface shader (e.g: 'normal' of <standard_surface>)
+		// if we come from a TS input, all geompropvalue along the path needs to be set in TS
+		// since we can only set at init-time the value of a TGuardValue and we're in a recursive state,
+		// if bTangentSpaceInput is false we set its value to IsTangentSpace otherwise we keep it as is (meaning true)
+		TGuardValue<bool> GuardInput(bTangentSpaceInput, bTangentSpaceInput || IsTangentSpaceInput(UsdInput));
+
+		mx::InputPtr Input = GetInput(Node, UsdInput);
+
+		bool bHasValue = true;
+		TUsdStore<pxr::UsdShadeInput::SourceInfoVector> ConnectedSources = UsdInput.GetConnectedSources();
+
+		// At this point an input should end up in either of these 3 cases, with the precedence in that order:
+		// 1. it has an interfaceName, if that's the case we continue to the next input (it has no value, no nodename)		
+		// 2. it is connected to a nodename with a valid nodedef (it has no value)
+		// 3. it has a value, if not defined by the input, the default one from the nodedef is taken
+		
+		//Let's loop over the connected sources, even though we should expect only one
+		for (const pxr::UsdShadeConnectionSourceInfo& Source : *ConnectedSources)
+		{
+			if (SetInterfaceName(Input, Source))
+			{
+				bHasValue = false;
+				break;
+			}
+
+			bool bIsNodeAlreadyExist;
+			mx::NodePtr ConnectedNode = GetNode(Source, bIsNodeAlreadyExist);
+			// Let's set if the input has an interfaceName
+			if (!ConnectedNode)
+			{
+				break;
+			}
+
+			Input->setConnectedNode(ConnectedNode);
+
+			// add the output of the connected node to the current input
+			TUsdStore<std::vector<pxr::UsdShadeOutput>> UsdOutputs = Source.source.GetOutputs();
+			for (const pxr::UsdShadeOutput& UsdOutput : *UsdOutputs)
+			{
+				// no need to add the default output
+				if (std::string OutputName = UsdOutput.GetBaseName().GetString(); OutputName != "out")
+				{
+					Input->setOutputString(OutputName);
+				}
+			}
+
+			// recurse if it's a newly created node, otherwise the work has already been done
+			if(!bIsNodeAlreadyExist)
+			{
+				bHasValue = ComputeShaderGraph(ConnectedNode, Source.source);
+			}
+
+			bHasValue = false;
+		}
+
+		if (bHasValue)
+		{
+			SetInputValue(Input, UsdInput);
+		}
+	}
+
+	return true;
+}
+
+bool FUsdMaterialXShaderGraph::IsGeomColor(const pxr::UsdShadeShader& GeomPropValueShader)
+{
+	if (pxr::UsdShadeInput GeomPropInput = GeomPropValueShader.GetInput(pxr::TfToken{ "geomprop" }))
+	{
+		if (std::string Value; GeomPropInput.Get<std::string>(&Value))
+		{
+			if (Value == "displayColor" || Value == "displayOpacity")
+			{
+				return true;
+			}
+		}
+	}
+	return false;
+}
+
+void FUsdMaterialXShaderGraph::SetInputValue(mx::InputPtr Input, const pxr::UsdShadeInput& UsdInput)
+{
+	auto SetInputValue = [&UsdInput, Input](auto Type)
+	{
+		using MtlxType = decltype(Type);
+
+		using FPairingUsdMtlxType = FPairingUsdMtlxType_t<MtlxType>;
+
+		if (FPairingUsdMtlxType Value; UsdInput.Get<FPairingUsdMtlxType>(&Value))
+		{
+			if constexpr (std::is_same_v<FPairingUsdMtlxType, pxr::GfVec4f> || std::is_same_v<FPairingUsdMtlxType, pxr::GfVec3f> || std::is_same_v<FPairingUsdMtlxType, pxr::GfVec2f>)
+			{
+				Input->setValue(MtlxType{ Value.data(), Value.data() + FPairingUsdMtlxType::dimension });
+			}
+			else
+			{
+				Input->setValue(Value);
+			}
+		}
+		else if(mx::NodePtr Node = Input->getParent()->asA<mx::Node>())
+		{
+			mx::NodeDefPtr NodeDef = Node->getNodeDef();
+			mx::InputPtr ActiveInput = NodeDef->getActiveInput(Input->getName());
+			MtlxType DefaultValue = ActiveInput->getDefaultValue()->asA<MtlxType>();
+			Input->setValue(DefaultValue);
+		}
+	};
+
+	const std::string& InputType = Input->getType();
+	if (InputType == "float")
+	{
+		SetInputValue(float{});
+	}
+	else if (InputType == "integer")
+	{
+		SetInputValue(int32{});
+	}
+	else if (InputType == "vector2")
+	{
+		SetInputValue(mx::Vector2{});
+	}
+	else if (InputType == "vector3")
+	{
+		SetInputValue(mx::Vector3{});
+	}
+	else if (InputType == "vector4")
+	{
+		SetInputValue(mx::Vector4{});
+	}
+	else if (InputType == "color3")
+	{
+		SetInputValue(mx::Color3{});
+	}
+	else if (InputType == "color4")
+	{
+		SetInputValue(mx::Color4{});
+	}
+	else if (InputType == "string")
+	{
+		if (TUsdStore<std::string> Value; UsdInput.Get<std::string>(&Value.Get()))
+		{
+			Input->setValue(Value.Get());
+		}
+		else if(mx::NodePtr Node = Input->getParent()->asA<mx::Node>())
+		{
+			mx::NodeDefPtr NodeDef = Node->getNodeDef();
+			mx::InputPtr ActiveInput = NodeDef->getActiveInput(Input->getName());
+			std::string DefaultValue = ActiveInput->getDefaultValue()->asA<std::string>();
+			Input->setValue(DefaultValue);
+		}
+	}
+	else if (InputType == "filename")
+	{
+		FString PathToResolve;
+		if(pxr::SdfAssetPath AssetPath; UsdInput.Get<pxr::SdfAssetPath>(&AssetPath))
+		{
+			TUsdStore<UE::FSdfLayer> Layer = UsdUtils::FindLayerForAttribute(UsdInput.GetAttr(), pxr::UsdTimeCode::Default().GetValue());
+
+			PathToResolve = Layer.Get()
+				? UsdToUnreal::ConvertString(TUsdStore<std::string>(pxr::SdfComputeAssetPathRelativeToLayer(Layer.Get(), AssetPath.GetAssetPath())).Get())
+				: UsdToUnreal::ConvertString(AssetPath.GetAssetPath());
+		}
+		Input->setValueString(std::string{ TCHAR_TO_UTF8(*PathToResolve) });
+	}
+	else
+	{
+		UE_LOG(LogUsd, Warning, TEXT("Couldn't find a value type for (%s)."), UTF8_TO_TCHAR(InputType.c_str()));
+	}
+}
+
+bool FUsdMaterialXShaderGraph::SetInterfaceName(MaterialX::InputPtr Input, const pxr::UsdShadeConnectionSourceInfo& Source)
+{
+	bool bIsSet = false;
+	const std::string& SourceName= Source.sourceName.GetString();
+	if (InterfaceNames.find(SourceName) != InterfaceNames.end())
+	{
+		if (mx::InputPtr InputInterfaceName = Document->getInput(SourceName))
+		{
+			Input->setInterfaceName(InputInterfaceName->getName());
+			bIsSet = true;
+		}
+	}
+	return bIsSet;
+}
+
+mx::NodePtr FUsdMaterialXShaderGraph::GetNode(const pxr::UsdShadeConnectionSourceInfo& Source, bool& bIsNodeAlreadyExist)
+{
+	bool bIsGeomPropValue;
+	bIsNodeAlreadyExist = true;
+	mx::NodeDefPtr NodeDef = GetNodeDef(Source, bIsGeomPropValue);
+
+	if (!NodeDef)
+	{
+		return nullptr;
+	}
+
+	std::string ConnectedShaderName = Source.source.GetPrim().GetName().GetString();
+
+	mx::NodePtr ConnectedNode = Document->getNode(ConnectedShaderName);
+	
+	// if the node is not present in the graph, we create it and notify that it's a new node
+	if (!ConnectedNode)
+	{
+		ConnectedNode = Document->addNodeInstance(NodeDef, ConnectedShaderName);
+		if(bIsGeomPropValue)
+		{
+			ConnectedNode->setTypedAttribute("UE:GeomPropImage", bIsGeomPropValue);
+		}
+		bIsNodeAlreadyExist = false;
+	}
+
+	return ConnectedNode;
+}
+
+mx::InputPtr FUsdMaterialXShaderGraph::GetInput(mx::NodePtr Node, const pxr::UsdShadeInput& UsdInput)
+{
+	std::string InputName = UsdInput.GetBaseName().GetString();
+	std::string InputType = UsdInput.GetTypeName().GetAsToken().GetString();
+
+	bool bMatch = GetMatchingInputType(InputName, InputType);
+
+	if (!bMatch && InputName == "geomprop")
+	{
+		if (Node->getCategory() == "geomcolor")
+		{
+			InputName = "index";
+			InputType = "integer";
+		}
+		else
+		{
+			InputName = "file";
+			InputType = "filename";
+			TUsdStore<std::string> Value;
+			UsdInput.Get<std::string>(&Value.Get());
+			if (!GeomPropValueNames.FindByPredicate([&Value](const FGeomProp& Other)
+				{
+					return Other.Name == Value->c_str();
+				}))
+			{
+				// we only set the the input as being in Tangent Space if also its type is a vec3
+				GeomPropValueNames.Emplace(Value->c_str(), bTangentSpaceInput && Node->getType() == "vector3" ? bTangentSpaceInput : false);
+			}
+		}
+	}
+
+	return Node->addInput(InputName, InputType);
+}
+
+mx::NodeDefPtr FUsdMaterialXShaderGraph::GetNodeDef(const pxr::UsdShadeConnectionSourceInfo& Source, bool& bIsGeomPropValue)
+{
+	pxr::UsdShadeShader UsdConnectedShader{ Source.source.GetPrim() };
+
+	if (!UsdConnectedShader)
+	{		                                                                                      
+		UE_LOG(LogUsd, Warning, TEXT("The '%s' connected source is not a valid USD shader."), *UsdToUnreal::ConvertString(Source.source.GetPrim().GetName()));
+	}
+
+	pxr::TfToken NodeDefTokenId;
+	UsdConnectedShader.GetShaderId(&NodeDefTokenId);
+
+	std::string NodeDefString = NodeDefTokenId.GetString();
+	bIsGeomPropValue = false;
+	if (std::size_t Pos = NodeDefString.find("geompropvalue"); Pos != std::string::npos)
+	{
+		// we convert a geomprop either to an image or a geomcolor in case if the name is displayColor
+		std::string NodeName = "image";
+		if (IsGeomColor(UsdConnectedShader))
+		{
+			NodeName = "geomcolor";
+		}
+		else
+		{
+			bIsGeomPropValue = true;
+			// Here we just replace the replace nodedef of a <geompropvalue> by the <image> one
+			// e.g: ND_geompropvalue_vector3
+			//		ND_image_vector3
+			// for integers/bool, <image> has no such node definitions, we just consider it as floats
+			// e.g: ND_geompropvalue_integer
+			//		ND_image_float
+			if (std::size_t pos = NodeDefString.rfind("integer"); pos != std::string::npos)
+			{
+				NodeDefString.replace(pos, sizeof("integer") - 1, "float");
+			}
+			else if (pos = NodeDefString.rfind("boolean"); pos != std::string::npos)
+			{
+				NodeDefString.replace(pos, sizeof("boolean") - 1, "float");
+			}
+			else if (pos = NodeDefString.rfind("string"); pos != std::string::npos)
+			{
+				UE_LOG(LogUsd, Warning, TEXT("'$s': '%s' are not supported"),
+					   *UsdToUnreal::ConvertString(Source.source.GetPrim().GetName()),
+					   *UsdToUnreal::ConvertString(NodeDefString));
+				return nullptr;
+			}
+		}
+		NodeDefString.replace(Pos, sizeof("geompropvalue") - 1, NodeName);
+	}
+
+	mx::NodeDefPtr NodeDef = Document->getNodeDef(NodeDefString);
+	
+	if(!NodeDef)
+	{
+		UE_LOG(LogUsd, Warning, TEXT("Couldn't find a nodedef for (%s)."), UTF8_TO_TCHAR(NodeDefTokenId.GetString().c_str()));
+	}
+
+	return NodeDef;
+}
+
+bool FUsdMaterialXShaderGraph::GetMatchingInputType(const std::string& InputName, std::string& InOutInputType)
+{
+	bool bMatch = false;
+
+	if (MaterialXTypes.find(InOutInputType) == MaterialXTypes.end())
+	{
+		if (auto It = UsdToMaterialXTypes.find(InOutInputType); It != UsdToMaterialXTypes.end())
+		{
+			InOutInputType = It->second;
+			bMatch = true;
+		}
+		else if (InputName == "file") // we're dealing with a texture, USD calls the type "asset" for the input
+		{
+			InOutInputType = "filename";
+			bMatch = true;
+		}
+	}
+
+	return bMatch;
+}
+
+bool FUsdMaterialXShaderGraph::IsTangentSpaceInput(const pxr::UsdShadeInput& UsdInput) const
+{
+	return TangentSpaceInputs.find(UsdInput.GetBaseName().GetString()) != TangentSpaceInputs.cend();
+}
+#endif // ENABLE_USD_MATERIALX
\ No newline at end of file
diff --git a/Engine/Plugins/Runtime/USDCore/Source/USDUtilities/Private/USDShadeConversion.cpp b/Engine/Plugins/Runtime/USDCore/Source/USDUtilities/Private/USDShadeConversion.cpp
--- a/Engine/Plugins/Runtime/USDCore/Source/USDUtilities/Private/USDShadeConversion.cpp
+++ b/Engine/Plugins/Runtime/USDCore/Source/USDUtilities/Private/USDShadeConversion.cpp
@@ -2123,331 +2123,6 @@ bool UsdToUnreal::FTextureParameterValue::GetSRGBValue() const
 	return bSRGB.Get(bFallbackSRGB);
 }
 
-#if WITH_EDITOR && PLATFORM_WINDOWS
-namespace UE::UsdShadeConversion::Private
-{
-	// clang-format off
-	std::unordered_set<std::string> MaterialXTypes
-	{
-		"boolean",
-		"integer",
-		"float",
-		"color3",
-		"color4",
-		"vector2",
-		"vector3",
-		"vector4",
-		"matrix33",
-		"matrix44",
-		"string",
-		"filename",
-		"geomname",
-		"surfaceshader",
-		"displacementshader",
-		"volumeshader",
-		"lightshader",
-		"material",
-		"none",
-		"integerarray",
-		"floatarray",
-		"color3array",
-		"color4array",
-		"vector2array",
-		"vector3array",
-		"vector4array",
-		"stringarray",
-		"geomnamearray",
-	};
-
-	std::unordered_map<std::string, std::string> UsdToMaterialXTypes
-	{
-		{"color3f", "color3"},
-		{"color4f", "color4"},
-		{"float2", "vector2"},
-		{"vector3f", "vector3"},
-		{"vector4f", "vector4"},
-		{"int", "integer"},
-	};
-	// clang-format on
-
-	bool ComputeShaderGraph(mx::NodePtr Node, const pxr::UsdShadeShader& Shader)
-	{
-		if (!Node)
-		{
-			return false;
-		}
-
-		mx::DocumentPtr Document = Node->getDocument();
-		TUsdStore<std::vector<pxr::UsdShadeInput>> UsdInputs = Shader.GetInputs();
-
-		for (const pxr::UsdShadeInput& UsdInput : *UsdInputs)
-		{
-			std::string InputName = UsdInput.GetBaseName().GetString();
-			std::string InputType = UsdInput.GetTypeName().GetAsToken().GetString();
-
-			if (MaterialXTypes.find(InputType) == MaterialXTypes.end())
-			{
-				if (auto It = UsdToMaterialXTypes.find(InputType); It != UsdToMaterialXTypes.end())
-				{
-					InputType = It->second;
-				}
-				else if (InputName == "file")	 // we're dealing with a texture, USD calls the type "asset" for the input
-				{
-					InputType = "filename";
-				}
-				else
-				{
-					UE_LOG(LogUsd, Warning, TEXT("Couldn't find a supported type for (%s)."), UTF8_TO_TCHAR(InputType.c_str()));
-				}
-			}
-
-			mx::InputPtr Input = Node->addInput(InputName, InputType);
-
-			TUsdStore<pxr::UsdShadeInput::SourceInfoVector> ConnectedSources = UsdInput.GetConnectedSources();
-			bool bHasValue = true;
-			for (const pxr::UsdShadeConnectionSourceInfo& Source : *ConnectedSources)
-			{
-				bHasValue = false;
-				pxr::UsdShadeShader ConnectedShader{Source.source.GetPrim()};
-
-				pxr::TfToken NodeDefTokenId;
-				ConnectedShader.GetShaderId(&NodeDefTokenId);
-				mx::NodeDefPtr NodeDef = Document->getNodeDef(NodeDefTokenId.GetString());
-
-				// if we couldn't find a nodedef just stop the graph here and propagate nullptr
-				if (!NodeDef)
-				{
-					Node = nullptr;
-					UE_LOG(LogUsd, Warning, TEXT("Couldn't find a nodedef for (%s)."), UTF8_TO_TCHAR(NodeDefTokenId.GetString().c_str()));
-					return false;
-				}
-
-				std::string ConnectedShaderName = ConnectedShader.GetPrim().GetName().GetString();
-
-				bool bComputeShaderGraph = false;
-				mx::NodePtr ConnectedNode = Document->getNode(ConnectedShaderName);
-				if (!ConnectedNode)
-				{
-					ConnectedNode = Document->addNodeInstance(NodeDef, ConnectedShaderName);
-					bComputeShaderGraph = true;
-				}
-
-				Input->setConnectedNode(ConnectedNode);
-
-				// add the output of the connected node to the current input
-				TUsdStore<std::vector<pxr::UsdShadeOutput>> UsdOutputs = ConnectedShader.GetOutputs();
-				for (const pxr::UsdShadeOutput& UsdOutput : *UsdOutputs)
-				{
-					// no need to add the default output
-					if (std::string OutputName = UsdOutput.GetBaseName().GetString(); OutputName != "out")
-					{
-						Input->setOutputString(OutputName);
-					}
-				}
-
-				if (bComputeShaderGraph)
-				{
-					ComputeShaderGraph(ConnectedNode, ConnectedShader);
-				}
-			}
-
-			if (bHasValue)
-			{
-				auto SetInputValue = [&UsdInput, Node, Input](auto Type)
-				{
-					using MtlxType = decltype(Type);
-
-					// complex template to just retrieve the right type, basically a mapping between Usd and Mtlx for Vector types, otherwise we just
-					// take the type (float, int, string)
-					// clang-format off
-					using UsdType = std::conditional_t<
-						std::is_same_v<MtlxType, mx::Color4> || std::is_same_v<MtlxType, mx::Vector4>,
-						pxr::GfVec4f,
-							std::conditional_t<
-							std::is_same_v<MtlxType, mx::Color3> || std::is_same_v<MtlxType, mx::Vector3>,
-								pxr::GfVec3f,
-								std::conditional_t<
-									std::is_same_v<MtlxType, mx::Vector2>,
-									pxr::GfVec2f,
-						MtlxType>>>;
-					// clang-format on
-
-					if (UsdType Value; UsdInput.Get<UsdType>(&Value))
-					{
-						if constexpr (std::is_same_v<UsdType, pxr::GfVec4f> || std::is_same_v<UsdType, pxr::GfVec3f>
-									  || std::is_same_v<UsdType, pxr::GfVec2f>)
-						{
-							Input->setValue(MtlxType{Value.data(), Value.data() + UsdType::dimension});
-						}
-						else
-						{
-							Input->setValue(Value);
-						}
-					}
-					else
-					{
-						mx::NodeDefPtr NodeDef = Node->getNodeDef();
-						mx::InputPtr ActiveInput = NodeDef->getActiveInput(Input->getName());
-						MtlxType DefaultValue = ActiveInput->getDefaultValue()->asA<MtlxType>();
-						Input->setValue(DefaultValue);
-					}
-				};
-
-				if (InputType == "float")
-				{
-					SetInputValue(float{});
-				}
-				else if (InputType == "integer")
-				{
-					SetInputValue(int32{});
-				}
-				else if (InputType == "vector2")
-				{
-					SetInputValue(mx::Vector2{});
-				}
-				else if (InputType == "vector3")
-				{
-					SetInputValue(mx::Vector3{});
-				}
-				else if (InputType == "vector4")
-				{
-					SetInputValue(mx::Vector4{});
-				}
-				else if (InputType == "color3")
-				{
-					SetInputValue(mx::Color3{});
-				}
-				else if (InputType == "color4")
-				{
-					SetInputValue(mx::Color4{});
-				}
-				else if (InputType == "string")
-				{
-					if (TUsdStore<std::string> Value; UsdInput.Get<std::string>(&Value.Get()))
-					{
-						Input->setValue(Value.Get());
-					}
-					else
-					{
-						mx::NodeDefPtr NodeDef = Node->getNodeDef();
-						mx::InputPtr ActiveInput = NodeDef->getActiveInput(Input->getName());
-						std::string DefaultValue = ActiveInput->getDefaultValue()->asA<std::string>();
-						Input->setValue(DefaultValue);
-					}
-				}
-				else if (InputType == "filename")
-				{
-					pxr::SdfAssetPath AssetPath;
-					UsdInput.Get<pxr::SdfAssetPath>(&AssetPath);
-					TUsdStore<UE::FSdfLayer> Layer = UsdUtils::FindLayerForAttribute(UsdInput.GetAttr(), pxr::UsdTimeCode::Default().GetValue());
-
-					FString PathToResolve = Layer.Get() ? UsdToUnreal::ConvertString(
-															  TUsdStore<std::string>(
-																  pxr::SdfComputeAssetPathRelativeToLayer(Layer.Get(), AssetPath.GetAssetPath())
-															  )
-																  .Get()
-														  )
-														: AssetPath.GetAssetPath().c_str();
-					Input->setValue(std::string{TCHAR_TO_UTF8(*PathToResolve)});
-				}
-				else
-				{
-					UE_LOG(LogUsd, Warning, TEXT("Couldn't find a value type for (%s)."), UTF8_TO_TCHAR(InputType.c_str()));
-				}
-			}
-		}
-
-		return true;
-	}
-}	 // namespace UE::UsdShadeConversion::Private
-
-mx::DocumentPtr UsdToUnreal::CreateSurfaceMaterial(const pxr::UsdPrim& UsdShadeMaterialPrim, const TCHAR* RenderContext)
-{
-	pxr::UsdShadeMaterial UsdShadeMaterial{UsdShadeMaterialPrim};
-
-	if (!UsdShadeMaterial)
-	{
-		return nullptr;
-	}
-
-	pxr::TfToken RenderContextToken = pxr::UsdShadeTokens->universalRenderContext;
-	if (RenderContext)
-	{
-		pxr::TfToken ProvidedRenderContextToken = UnrealToUsd::ConvertToken(RenderContext).Get();
-		RenderContextToken = ProvidedRenderContextToken;
-	}
-
-	pxr::UsdShadeShader SurfaceShader = UsdShadeMaterial.ComputeSurfaceSource(RenderContextToken);
-	pxr::UsdShadeShader DisplacementShader = UsdShadeMaterial.ComputeDisplacementSource(RenderContextToken);
-
-	if (!SurfaceShader)
-	{
-		return nullptr;
-	}
-
-	mx::DocumentPtr Document = mx::createDocument();
-	{
-		mx::DocumentPtr MaterialXLibrary = mx::createDocument();
-		mx::FileSearchPath MaterialXFolder{
-			TCHAR_TO_UTF8(*FPaths::Combine(FPaths::EngineDir(), TEXT("Binaries"), TEXT("ThirdParty"), TEXT("MaterialX")))
-		};
-		mx::StringSet LoadedLibs = mx::loadLibraries({"libraries"}, MaterialXFolder, MaterialXLibrary);
-		Document->importLibrary(MaterialXLibrary);
-	}
-
-	auto CreateShaderGraph = [](const pxr::UsdShadeShader& SurfaceShader, mx::DocumentPtr Document) -> mx::NodePtr
-	{
-		pxr::TfToken NodeDefTokenId;
-		SurfaceShader.GetShaderId(&NodeDefTokenId);
-		mx::NodeDefPtr NodeDef = Document->getNodeDef(NodeDefTokenId.GetString());
-
-		// no node definition, just return
-		if (!NodeDef)
-		{
-			return nullptr;
-		}
-
-		mx::NodePtr MtlxSurfaceShader = Document->addNodeInstance(NodeDef, SurfaceShader.GetPrim().GetName().GetString());
-
-		// if the shadergraph is not valid (probably caused by a missing nodedef or wrong type), just invalide it all
-		if (!UE::UsdShadeConversion::Private::ComputeShaderGraph(MtlxSurfaceShader, SurfaceShader))
-		{
-			MtlxSurfaceShader = nullptr;
-		}
-
-		return MtlxSurfaceShader;
-	};
-
-	// Create the shader graph for the surface shader and displacement shader
-	if (mx::NodePtr MaterialXSurfaceShader = CreateShaderGraph(SurfaceShader, Document))
-	{
-		mx::NodeDefPtr SurfaceMaterialNodeDef = Document->getNodeDef("ND_surfacematerial");
-		mx::NodePtr SurfaceMaterial = Document->addNodeInstance(SurfaceMaterialNodeDef, UsdShadeMaterial.GetPrim().GetName().GetString());
-
-		SurfaceMaterial->setConnectedNode("surfaceshader", MaterialXSurfaceShader);
-
-		// the displacement shader is not necessarily required
-		if (DisplacementShader)
-		{
-			if (mx::NodePtr MaterialXDisplacementShader = CreateShaderGraph(DisplacementShader, Document); MaterialXDisplacementShader)
-			{
-				SurfaceMaterial->setConnectedNode("displacementshader", MaterialXDisplacementShader);
-			}
-		}
-
-		std::string ErrorMessage;
-		SurfaceMaterial->validate(&ErrorMessage);
-		if (!ErrorMessage.empty())
-		{
-			UE_LOG(LogUsd, Error, TEXT("USD MaterialX: %s."), UTF8_TO_TCHAR(ErrorMessage.c_str()));
-			SurfaceMaterial = nullptr;
-		}
-	}
-
-	return Document;
-}
-#endif	  // WITH_EDITOR && PLATFORM_WINDOWS
-
 bool UsdToUnreal::ConvertMaterial(
 	const pxr::UsdPrim& UsdShadeMaterialPrim,
 	FUsdPreviewSurfaceMaterialData& OutMaterialData,
diff --git a/Engine/Plugins/Runtime/USDCore/Source/USDUtilities/Public/USDMaterialXShaderGraph.h b/Engine/Plugins/Runtime/USDCore/Source/USDUtilities/Public/USDMaterialXShaderGraph.h
new file mode 100644
--- /dev/null
+++ b/Engine/Plugins/Runtime/USDCore/Source/USDUtilities/Public/USDMaterialXShaderGraph.h
@@ -0,0 +1,118 @@
+// Copyright Epic Games, Inc. All Rights Reserved.
+
+#pragma once
+
+#include "Containers/UnrealString.h"
+
+#if WITH_EDITOR && PLATFORM_WINDOWS
+#include "MaterialXFormat/Util.h"
+
+#if USE_USD_SDK
+#include "USDIncludesStart.h"
+#include "pxr/base/tf/token.h"
+#include "pxr/usd/sdf/layerUtils.h"
+#include "pxr/usd/usdShade/material.h"
+#include "USDIncludesEnd.h"
+
+#define ENABLE_USD_MATERIALX
+#endif // USE_USD_SDK
+#endif //WITH_EDITOR && PLATFORM_WINDOWS
+
+/** 
+ * Create a MaterialX shader graph from a UsdPrim, all <geompropvalue> nodes are either converted to an <image> or a proper <geomcolor>
+ * Just get the Document in order to search for the surface material, by just looking for the name of the Prim it will return the correct surface material.
+ */
+class USDUTILITIES_API FUsdMaterialXShaderGraph
+{
+public:
+
+	struct FGeomProp
+	{
+		FString Name;
+		bool bTangentSpace = false;
+	};
+
+#ifdef ENABLE_USD_MATERIALX
+
+	/**
+	 * @param UsdShadeMaterialPrim - The prim used to generate the shader graph, should be a valid pxr::UsdShadeMaterial
+	 * @param RenderContext - Which render context output to read from the UsdShadeMaterial
+	 */
+	FUsdMaterialXShaderGraph(const pxr::UsdPrim& MaterialPrim, const TCHAR* RenderContext);
+
+	MaterialX::DocumentPtr GetDocument() const;
+
+	/** Return the names of the <geompropvalue> nodes which have converted to <image> node*/
+	const TArray<FGeomProp>& GetGeomPropValueNames() const;
+
+private:
+
+	void CreateSurfaceMaterial(const pxr::UsdShadeMaterial& UsdShadeMaterial, const TCHAR* RenderContext);
+
+	bool ComputeShaderGraph(MaterialX::NodePtr Node, const pxr::UsdShadeShader& Shader);
+
+	bool IsGeomColor(const pxr::UsdShadeShader& GeomPropValueShader);
+
+	void SetInputValue(MaterialX::InputPtr Input, const pxr::UsdShadeInput& UsdInput);
+
+	/** Set the interfaceName of the input if any */ 
+	bool SetInterfaceName(MaterialX::InputPtr Input, const pxr::UsdShadeConnectionSourceInfo& Source);
+
+	/**
+	 * Get the MaterialX node from a UsdShadeConnectionSourceInfo
+	 * @param Source - upstream source in the Shader Graph, usually the connected node to an input
+	 * @param bIsNodeAlreadyExist - return true if the node was already present in the shader graph, false otherwise
+	 * @return The corresponding MaterialX node in the shader graph, if it doesn't exist it creates one
+	 */
+	MaterialX::NodePtr GetNode(const pxr::UsdShadeConnectionSourceInfo& Source, bool& bIsNodeAlreadyExist);
+
+	MaterialX::InputPtr GetInput(MaterialX::NodePtr Node, const pxr::UsdShadeInput& UsdInput);
+
+	MaterialX::NodeDefPtr GetNodeDef(const pxr::UsdShadeConnectionSourceInfo& Source, bool& bIsGeomPropValue);
+
+	/**
+	 * Given an input type, retrieve the MaterialX associated type
+	 * @param InputName - The name of the input
+	 * @param InOutInputType - The type of the input, if it's valid it will now be the MaterialX type
+	 * @return true if there is a valid match between USD and MaterialX or if it's already a MaterialX type, false otherwise
+	 */
+	bool GetMatchingInputType(const std::string& InputName, std::string& InOutInputType);
+
+	bool IsTangentSpaceInput(const pxr::UsdShadeInput& UsdInput) const;
+
+private:
+
+	std::unordered_set<std::string> MaterialXTypes;
+
+	std::unordered_map<std::string, std::string> UsdToMaterialXTypes;
+
+	// Names of the inputs from a <surfaceshader> that requires to be in tangent space for the baking of the primvars later on
+	std::unordered_set<std::string> TangentSpaceInputs;
+
+	// MaterialX interfacenames (USD calls that InterfaceInputs), some inputs use them instead of a value or a direct connection to a node
+	std::unordered_set<std::string> InterfaceNames; 
+
+	MaterialX::DocumentPtr Document;
+
+	TArray<FGeomProp> GeomPropValueNames;
+
+	bool bTangentSpaceInput = false;
+
+private:
+
+	// complex template to just retrieve the right type, basically a mapping between Usd and Mtlx for Vector types, otherwise we just
+	// take the type (float, int, string)
+	template<typename MtlxType>
+	using FPairingUsdMtlxType_t =
+		std::conditional_t<
+		   std::is_same_v<MtlxType, MaterialX::Color4> || std::is_same_v<MtlxType, MaterialX::Vector4>,
+		   pxr::GfVec4f,
+		      std::conditional_t<
+		         std::is_same_v<MtlxType, MaterialX::Color3> || std::is_same_v<MtlxType, MaterialX::Vector3>,
+		         pxr::GfVec3f,
+		            std::conditional_t<
+		               std::is_same_v<MtlxType, MaterialX::Vector2>,
+		               pxr::GfVec2f,
+		   MtlxType>>>;
+#endif // ENABLE_USD_MATERIALX
+};
\ No newline at end of file
diff --git a/Engine/Plugins/Runtime/USDCore/Source/USDUtilities/Public/USDShadeConversion.h b/Engine/Plugins/Runtime/USDCore/Source/USDUtilities/Public/USDShadeConversion.h
--- a/Engine/Plugins/Runtime/USDCore/Source/USDUtilities/Public/USDShadeConversion.h
+++ b/Engine/Plugins/Runtime/USDCore/Source/USDUtilities/Public/USDShadeConversion.h
@@ -134,15 +134,6 @@ namespace UsdToUnreal
 		bool bShareAssetsForIdenticalPrims = true
 	);
 
-	/** Create a MaterialX Surface material. If the created shader graph is not valid according to MaterialX spec, return nullptr
-	 * We return the Document in order for the caller to search for the surface material, by just looking for the name will return the correct surface material.
-	 * We need to return the document instead of the node directly otherwise the Document will go out of scope and create an orphaned shader graph
-	 * @param UsdShadeMaterialPrim - The prim used to generate the shader graph, should be a valid pxr::UsdShadeMaterial
-	 * @param RenderContext - Which render context output to read from the UsdShadeMaterial
-	 */
-#if WITH_EDITOR && PLATFORM_WINDOWS
-	USDUTILITIES_API MaterialX::DocumentPtr CreateSurfaceMaterial(const pxr::UsdPrim& UsdShadeMaterialPrim, const TCHAR* RenderContext);
-#endif
 	/**
 	 * Attemps to assign the values of the surface shader inputs to the MaterialInstance parameters by matching the inputs display names to the
 	 * parameters names.
