[USD] Fixing IsCollepsable checks (regarding Animations not showing up on import)

#rnx

diff --git a/Engine/Plugins/Runtime/USDCore/Source/USDUtilities/Private/Objects/USDInfoCache.cpp b/Engine/Plugins/Runtime/USDCore/Source/USDUtilities/Private/Objects/USDInfoCache.cpp
--- a/Engine/Plugins/Runtime/USDCore/Source/USDUtilities/Private/Objects/USDInfoCache.cpp
+++ b/Engine/Plugins/Runtime/USDCore/Source/USDUtilities/Private/Objects/USDInfoCache.cpp
@@ -740,8 +740,9 @@ namespace UE::UsdInfoCache::Private
 			TArray<pxr::UsdPrim> Children;
 			for (pxr::UsdPrim Child : UsdPrim.GetFilteredChildren(pxr::UsdTraverseInstanceProxies(pxr::UsdPrimAllPrimsPredicate)))
 			{
-				// We don't care about non-xformable prims (materials, stuff we don't have schema translators for, etc.)
-				if (Child.IsA<pxr::UsdGeomXformable>())
+				// We don't care about non-GeomImagable prims 
+				// (materials, etc., stuff we don't have schema translators for will be skipped and default to bCanBeCollapsed=true)
+				if (Child.IsA<pxr::UsdGeomImageable>())
 				{
 					Children.Emplace(Child);
 				}
