USD: Fixed missing transforms when importing USD geometry caches.

diff --git a/Engine/Plugins/Runtime/USDCore/Source/USDUtilities/Private/USDPrimConversion.cpp b/Engine/Plugins/Runtime/USDCore/Source/USDUtilities/Private/USDPrimConversion.cpp
--- a/Engine/Plugins/Runtime/USDCore/Source/USDUtilities/Private/USDPrimConversion.cpp
+++ b/Engine/Plugins/Runtime/USDCore/Source/USDUtilities/Private/USDPrimConversion.cpp
@@ -437,8 +437,8 @@ void UsdToUnreal::PropagateTransform(
 	}
 	else
 	{
-		// Leaf is not completely connected to Root by Xformables
-		OutTransform = FTransform::Identity;
+		// The Leaf is not an Xformable but we still want to propagate the transform to the Root
+		PropagateTransform(Stage, Root, Leaf.GetParent(), EvalTime, OutTransform);
 	}
 }
 
