[Interchange - USD] Adding StageInfo data into the Extra Information section of Interchange Import Dialog window.

#rnx

diff --git a/Engine/Plugins/Interchange/Extensions/OpenUSD/Source/Import/Private/InterchangeUsdTranslator.cpp b/Engine/Plugins/Interchange/Extensions/OpenUSD/Source/Import/Private/InterchangeUsdTranslator.cpp
--- a/Engine/Plugins/Interchange/Extensions/OpenUSD/Source/Import/Private/InterchangeUsdTranslator.cpp
+++ b/Engine/Plugins/Interchange/Extensions/OpenUSD/Source/Import/Private/InterchangeUsdTranslator.cpp
@@ -49,6 +49,7 @@
 #include "InterchangeUsdContext.h"
 #include "InterchangeVolumeNode.h"
 #include "Mesh/InterchangeMeshPayload.h"
+#include "Nodes/InterchangeSourceNode.h"
 #include "Nodes/InterchangeUserDefinedAttribute.h"
 #include "Usd/InterchangeUsdDefinitions.h"
 #include "Volume/InterchangeVolumeDefinitions.h"
@@ -4091,6 +4092,20 @@ namespace UE::InterchangeUsdTranslator::Private
 
 		return true;
 	}
+
+	void ProcessExtraInformation(UInterchangeBaseNodeContainer& NodeContainer, UE::FUsdStage Stage)
+	{
+		UInterchangeSourceNode* SourceNode = UInterchangeSourceNode::FindOrCreateUniqueInstance(&NodeContainer);
+
+		TMap<FString, FString> MetadataMap;
+		UsdUtils::ReadStageMetaData(Stage, MetadataMap);
+
+		for (TPair<FString, FString>& MetaDataEntry : MetadataMap)
+		{
+			SourceNode->SetExtraInformation(MetaDataEntry.Key, MetaDataEntry.Value);
+		}
+	}
+
 #else
 	class UInterchangeUSDTranslatorImpl
 	{
@@ -4228,6 +4243,8 @@ bool UInterchangeUSDTranslator::Translate(UInterchangeBaseNodeContainer& NodeCon
 
 			StageToImport.SetInterpolationType(Settings->InterpolationType);
 		}
+
+		ProcessExtraInformation(NodeContainer, StageToImport);
 	}
 
 	// Setup info cache
diff --git a/Engine/Plugins/Runtime/USDCore/Source/USDUtilities/Private/USDConversionUtils.cpp b/Engine/Plugins/Runtime/USDCore/Source/USDUtilities/Private/USDConversionUtils.cpp
--- a/Engine/Plugins/Runtime/USDCore/Source/USDUtilities/Private/USDConversionUtils.cpp
+++ b/Engine/Plugins/Runtime/USDCore/Source/USDUtilities/Private/USDConversionUtils.cpp
@@ -19,9 +19,11 @@
 #include "USDSkeletalDataConversion.h"
 #include "USDTypesConversion.h"
 #include "USDUnrealAssetInfo.h"
+#include "USDValueConversion.h"
 #include "UsdWrappers/SdfPath.h"
 #include "UsdWrappers/UsdPrim.h"
 #include "UsdWrappers/UsdStage.h"
+#include "UsdWrappers/VtValue.h"
 
 #include "Animation/AnimBlueprint.h"
 #include "Animation/AnimSequence.h"
@@ -64,6 +66,7 @@
 #include "pxr/base/tf/token.h"
 #include "pxr/usd/kind/registry.h"
 #include "pxr/usd/sdf/copyUtils.h"
+#include "pxr/usd/sdf/schema.h"
 #include "pxr/usd/usd/attribute.h"
 #include "pxr/usd/usd/editContext.h"
 #include "pxr/usd/usd/modelAPI.h"
@@ -4025,6 +4028,98 @@ void UsdUtils::CollectSchemaAnalytics(const UE::FUsdStage& Stage, const FString&
 {
 }
 
+void UsdUtils::ReadStageMetaData(UE::FUsdStage Stage, TMap<FString, FString>& OutMetaDataMap)
+{
+#if USE_USD_SDK
+	using namespace pxr;
+
+	const FUsdStageInfo StageInfo(Stage);
+	OutMetaDataMap.Add(TEXT("Meters Per Unit"), FString::SanitizeFloat(StageInfo.MetersPerUnit));
+	if (StageInfo.UpAxis == EUsdUpAxis::ZAxis)
+	{
+		OutMetaDataMap.Add(TEXT("Up Axis"), TEXT("Z"));
+	}
+	else if (StageInfo.UpAxis == EUsdUpAxis::YAxis)
+	{
+		OutMetaDataMap.Add(TEXT("Up Axis"), TEXT("Y"));
+	}
+
+	//DefaultPrim
+	UE::FVtValue DefaultPrim;
+	if (Stage.GetMetadata(*UsdToUnreal::ConvertToken(SdfFieldKeys->DefaultPrim), DefaultPrim))
+	{
+		VtValue UsdValue = DefaultPrim.GetUsdValue();
+
+		if (UsdValue.IsHolding<TfToken>())
+		{
+			TfToken Value = UsdValue.UncheckedGet<TfToken>();
+			FString DefaultPrimString = UsdToUnreal::ConvertToken(Value);
+
+			if (DefaultPrimString.Len() > 0)
+			{
+				OutMetaDataMap.Add(TEXT("Default Prim"), DefaultPrimString);
+			}
+		}
+	}
+
+	//double typed potential meta data handling:
+	auto ProcessDoubleTypedMetadata = [&OutMetaDataMap, &Stage](const FString& ExtraInformationKey, const FString& TokenStringified)
+		{
+			UE::FVtValue UEValue;
+			if (Stage.GetMetadata(*TokenStringified, UEValue))
+			{
+				VtValue UsdValue = UEValue.GetUsdValue();
+
+				if (UsdValue.IsHolding<double>())
+				{
+					double Value = UsdValue.UncheckedGet<double>();
+					OutMetaDataMap.Add(ExtraInformationKey, FString::SanitizeFloat(Value));
+				}
+			}
+		};
+	ProcessDoubleTypedMetadata(TEXT("Time Codes Per Second"), UsdToUnreal::ConvertToken(SdfFieldKeys->TimeCodesPerSecond));
+	ProcessDoubleTypedMetadata(TEXT("Frames Per Second"), UsdToUnreal::ConvertToken(SdfFieldKeys->FramesPerSecond));
+	ProcessDoubleTypedMetadata(TEXT("Start Time Code"), UsdToUnreal::ConvertToken(SdfFieldKeys->StartTimeCode));
+	ProcessDoubleTypedMetadata(TEXT("End Time Code"), UsdToUnreal::ConvertToken(SdfFieldKeys->EndTimeCode));
+
+	//Documentation
+	UE::FVtValue Documentation;
+	if (Stage.GetMetadata(*UsdToUnreal::ConvertToken(SdfFieldKeys->Documentation), Documentation))
+	{
+		VtValue UsdValue = Documentation.GetUsdValue();
+
+		if (UsdValue.IsHolding<std::string>())
+		{
+			std::string Value = UsdValue.UncheckedGet<std::string>();
+			OutMetaDataMap.Add(TEXT("Documentation"), UsdToUnreal::ConvertString(Value.c_str()));
+		}
+	}
+
+	//CustomLayerData
+	UE::FVtValue CustomLayerData;
+	if (Stage.GetMetadata(*UsdToUnreal::ConvertToken(SdfFieldKeys->CustomLayerData), CustomLayerData))
+	{
+		VtValue UsdValue = CustomLayerData.GetUsdValue();
+
+		if (UsdValue.IsHolding<VtDictionary>())
+		{
+			const VtDictionary& UsdDictionary = UsdValue.UncheckedGet<VtDictionary>();
+
+			for (VtDictionary::const_iterator ValueIter = UsdDictionary.begin(); ValueIter != UsdDictionary.end(); ++ValueIter)
+			{
+				const std::string& DictFieldName = ValueIter->first;
+				const VtValue& DictFieldValue = ValueIter->second;
+
+				FString StringifiedKey = UsdToUnreal::ConvertString(DictFieldName);
+				FString StringifiedValue = UsdUtils::Stringify(DictFieldValue);
+
+				OutMetaDataMap.Add(StringifiedKey, StringifiedValue);
+			}
+		}
+	}
+#endif
+}
+
 #if USE_USD_SDK
 #undef LOCTEXT_NAMESPACE
 #endif
diff --git a/Engine/Plugins/Runtime/USDCore/Source/USDUtilities/Public/USDConversionUtils.h b/Engine/Plugins/Runtime/USDCore/Source/USDUtilities/Public/USDConversionUtils.h
--- a/Engine/Plugins/Runtime/USDCore/Source/USDUtilities/Public/USDConversionUtils.h
+++ b/Engine/Plugins/Runtime/USDCore/Source/USDUtilities/Public/USDConversionUtils.h
@@ -655,4 +655,7 @@ namespace UsdUtils
 	/** Collects how many times each schema shows up on the provided stage and send it as an analytics event */
 	UE_DEPRECATED(5.5, "This function has been moved to the USDSchemasModule")
 	USDUTILITIES_API void CollectSchemaAnalytics(const UE::FUsdStage& Stage, const FString& EventName);
+
+	/** Parses and stores primary Stage metadata key-info pairs in the MetadataMap */
+	USDUTILITIES_API void ReadStageMetaData(UE::FUsdStage Stage, TMap<FString, FString>& OutMetaDataMap);
 }	 // namespace UsdUtils
