USD:
- Optimize memory consumption when storing the primvars
- Primvars are now stored per type in the MeshDescription instead of an array of Vector4

Interchange:
- Updating the code of the baking to reflect the new storage of the primvars by type in the MeshDescription

diff --git a/Engine/Plugins/Interchange/Extensions/OpenUSD/Source/Import/Private/InterchangeUSDPipeline.cpp b/Engine/Plugins/Interchange/Extensions/OpenUSD/Source/Import/Private/InterchangeUSDPipeline.cpp
--- a/Engine/Plugins/Interchange/Extensions/OpenUSD/Source/Import/Private/InterchangeUSDPipeline.cpp
+++ b/Engine/Plugins/Interchange/Extensions/OpenUSD/Source/Import/Private/InterchangeUSDPipeline.cpp
@@ -1045,61 +1045,86 @@ namespace UE::InterchangeUsdPipeline::Private
 					FactoryNode->GetStringAttribute(USD::Primvar::Name + FString::FromInt(Index), PrimvarName);
 					FactoryNode->GetBooleanAttribute(USD::Primvar::TangentSpace + FString::FromInt(Index), bTangentSpace);
 
-					TVertexInstanceAttributesRef<FVector4f> Primvar = MeshDescriptionPrimvar.VertexInstanceAttributes().GetAttributesRef<FVector4f>(
-						*PrimvarName
-					);
-
-					if (!(Primvar.IsValid() && Primvar.GetNumElements() == VertexColors.GetNumElements()
-						  && Primvar.GetNumElements() == Normals.GetNumElements()))
+					MeshDescriptionPrimvar.VertexInstanceAttributes().ForEach(
+					[this, &PrimvarName, &VertexColors, &Normals, TextureSampleExpression, &MeshDescriptionPrimvar, bTangentSpace]
+					(const FName AttributeName, auto AttributesConstRef)
 					{
-						continue;
-					}
+						if (AttributeName != PrimvarName)
+						{
+							return;
+						}
 
-					// We put the primvar in the Color/Normal channel, that way during the conversion to a DynamicMesh it will also handle the
-					// non-manifold case
-					for (const int32 VertexInstanceID : MeshDescriptionPrimvar.VertexInstances().GetElementIDs())
-					{
-						bTangentSpace ? Normals.Set(VertexInstanceID, Primvar.Get(VertexInstanceID))
-									  : VertexColors.Set(VertexInstanceID, Primvar.Get(VertexInstanceID));
-					}
+						if (!(AttributesConstRef.IsValid() && AttributesConstRef.GetNumElements() == VertexColors.GetNumElements()
+							&& AttributesConstRef.GetNumElements() == Normals.GetNumElements()))
+						{
+							return;
+						}
+
+						// We put the primvar in the Color/Normal channel, that way during the conversion to a DynamicMesh it will also handle the
+						// non-manifold case
+						for (const int32 VertexInstanceID : MeshDescriptionPrimvar.VertexInstances().GetElementIDs())
+						{
+							auto VertexInstanceValue = AttributesConstRef.Get(VertexInstanceID);
+
+							// underlying type of the Value in the AttributesConstRef
+							using T = std::decay_t<decltype(VertexInstanceValue)>;
+
+							if constexpr (std::is_same_v<T, bool> || std::is_same_v<T, int32> || std::is_same_v<T, float>)
+							{
+								// store the value inside a float to avoid the compile error: conversion from 'AttributeType' to 'T' requires a narrowing conversion
+								const float VertexInstanceValueFloat = VertexInstanceValue;
+								VertexColors.Set(VertexInstanceID, FVector4f{ VertexInstanceValueFloat, VertexInstanceValueFloat, VertexInstanceValueFloat, VertexInstanceValueFloat });
+							}
+							else if constexpr (std::is_same_v<T, FVector2f>)
+							{
+								VertexColors.Set(VertexInstanceID, FVector4f{ VertexInstanceValue, VertexInstanceValue });
+							}
+							else if constexpr (std::is_same_v<T, FVector3f> || std::is_same_v<T, FVector4f>)
+							{
+								bTangentSpace
+									? Normals.Set(VertexInstanceID, VertexInstanceValue)
+									: VertexColors.Set(VertexInstanceID, VertexInstanceValue);
+							}
+						}
 
-					FDynamicMesh3 DynamicMesh;
-					FMeshDescriptionToDynamicMesh Convert;
-					Convert.Convert(&MeshDescriptionPrimvar, DynamicMesh);
-
-					const FDynamicMeshAABBTree3 DetailSpatial(&DynamicMesh);
-					FMeshBakerDynamicMeshSampler DetailSampler(&DynamicMesh, &DetailSpatial);
-					TSharedPtr<FMeshPropertyMapEvaluator> PropertyEval = MakeShared<FMeshPropertyMapEvaluator>();
-					PropertyEval->Property = bTangentSpace ? EMeshPropertyMapType::Normal : EMeshPropertyMapType::VertexColor;
-
-					constexpr int32 TextureSize = 1024;
-					const FImageDimensions ImageDimensions(TextureSize, TextureSize);
-					FMeshMapBaker Baker;
-					Baker.SetCorrespondenceStrategy(FMeshBaseBaker::ECorrespondenceStrategy::Identity);
-					Baker.SetTargetMesh(&DynamicMesh);
-					Baker.SetDetailSampler(&DetailSampler);
-					Baker.AddEvaluator(PropertyEval);
-					Baker.SetTargetMeshUVLayer(UVChannel);
-					Baker.SetDimensions(ImageDimensions);
-					Baker.SetProjectionDistance(3.0f);
-					Baker.SetSamplesPerPixel(1);
-					Baker.SetFilter(FMeshMapBaker::EBakeFilterType::BSpline);
-					Baker.SetGutterEnabled(true);
-					Baker.SetGutterSize(4);
-					Baker.SetTileSize(TextureSize);
-					Baker.Bake();
-
-					FTexture2DBuilder TextureBuilder;
-					const FTexture2DBuilder::ETextureType TextureType = GetTextureType(TextureSampleExpression);
-
-					TextureBuilder
-						.InitializeAndReplaceExistingTexture(Cast<UTexture2D>(TextureSampleExpression->Texture.Get()), TextureType, ImageDimensions);
-					TextureBuilder.Copy(*Baker.GetBakeResults(0)[0]);
-					TextureBuilder.Commit();
-
-					TextureSampleExpression->ConstCoordinate = UVChannel;
-
-					MaterialExpressionTextureSamples.Add(TextureSampleExpression);
+						FDynamicMesh3 DynamicMesh;
+						FMeshDescriptionToDynamicMesh Convert;
+						Convert.Convert(&MeshDescriptionPrimvar, DynamicMesh);
+
+						const FDynamicMeshAABBTree3 DetailSpatial(&DynamicMesh);
+						FMeshBakerDynamicMeshSampler DetailSampler(&DynamicMesh, &DetailSpatial);
+						TSharedPtr<FMeshPropertyMapEvaluator> PropertyEval = MakeShared<FMeshPropertyMapEvaluator>();
+						PropertyEval->Property = bTangentSpace ? EMeshPropertyMapType::Normal : EMeshPropertyMapType::VertexColor;
+
+						constexpr int32 TextureSize = 1024;
+						const FImageDimensions ImageDimensions(TextureSize, TextureSize);
+						FMeshMapBaker Baker;
+						Baker.SetCorrespondenceStrategy(FMeshBaseBaker::ECorrespondenceStrategy::Identity);
+						Baker.SetTargetMesh(&DynamicMesh);
+						Baker.SetDetailSampler(&DetailSampler);
+						Baker.AddEvaluator(PropertyEval);
+						Baker.SetTargetMeshUVLayer(UVChannel);
+						Baker.SetDimensions(ImageDimensions);
+						Baker.SetProjectionDistance(3.0f);
+						Baker.SetSamplesPerPixel(1);
+						Baker.SetFilter(FMeshMapBaker::EBakeFilterType::BSpline);
+						Baker.SetGutterEnabled(true);
+						Baker.SetGutterSize(4);
+						Baker.SetTileSize(TextureSize);
+						Baker.Bake();
+
+						FTexture2DBuilder TextureBuilder;
+						const FTexture2DBuilder::ETextureType TextureType = GetTextureType(TextureSampleExpression);
+
+						TextureBuilder
+							.InitializeAndReplaceExistingTexture(Cast<UTexture2D>(TextureSampleExpression->Texture.Get()), TextureType, ImageDimensions);
+						TextureBuilder.Copy(*Baker.GetBakeResults(0)[0]);
+						TextureBuilder.Commit();
+
+						TextureSampleExpression->ConstCoordinate = UVChannel;
+
+						MaterialExpressionTextureSamples.Add(TextureSampleExpression);
+					});
 				}
 
 				UpdateMaterials();
diff --git a/Engine/Plugins/Runtime/USDCore/Source/USDUtilities/Private/USDGeomMeshConversion.cpp b/Engine/Plugins/Runtime/USDCore/Source/USDUtilities/Private/USDGeomMeshConversion.cpp
--- a/Engine/Plugins/Runtime/USDCore/Source/USDUtilities/Private/USDGeomMeshConversion.cpp
+++ b/Engine/Plugins/Runtime/USDCore/Source/USDUtilities/Private/USDGeomMeshConversion.cpp
@@ -1154,14 +1154,22 @@ namespace UE::UsdGeomMeshConversion::Private
 		TArray<pxr::VtArray<pxr::GfVec2f>> UVSets;
 
 		// Custom attributes, that may come from primvars (used essentially for baking textures in Interchange)
-		struct FCustomPrimvar
+		struct FPrimvar
 		{
-			pxr::VtArray<pxr::GfVec4f> Data; // Let's put any primvar type in a vec4, we'll convert later in the MeshDescription
+			using FPrimvarType = TVariant<
+				pxr::VtArray<bool>,
+				pxr::VtArray<int32>,
+				pxr::VtArray<float>,
+				pxr::VtArray<pxr::GfVec2f>,
+				pxr::VtArray<pxr::GfVec3f>,
+				pxr::VtArray<pxr::GfVec4f>>;
+
+			FPrimvarType Data;
 			pxr::VtArray<int> Indices;
 			pxr::TfToken Interpolation;
 		};
 
-		TMap<FString, FCustomPrimvar> CustomPrimvars;
+		TMap<FString, FPrimvar> Primvars;
 
 		pxr::TfToken PointInterpolation;
 		pxr::TfToken NormalInterpolation;
@@ -1200,6 +1208,14 @@ namespace UE::UsdGeomMeshConversion::Private
 		TArray<FUsdRawBoneInfluence> Influences;
 	};
 
+	// https://en.cppreference.com/w/cpp/utility/variant/visit2
+	// Helper type for the visitor
+	template<class... Ts>
+	struct TOverloaded : Ts... {using Ts::operator()...; };
+	//Explicit deduction guide (not needed as of C++20
+	template<class... Ts>
+	TOverloaded(Ts...) -> TOverloaded<Ts...>;
+
 	namespace USDSkeletalDataConversion
 	{
 		bool HasMultipleRootBones(const pxr::UsdSkelSkeletonQuery& SkeletonQuery)
@@ -1301,64 +1317,49 @@ namespace UE::UsdGeomMeshConversion::Private
 			}
 
 			std::string PrimvarType = Primvar.GetTypeName().GetAsToken().GetString();
-
-			FUsdMeshData::FCustomPrimvar CustomPrimvar;
-			auto CollectPrimvarData = [&PrimvarType, &Options, &CustomPrimvar]<class T>(const pxr::VtArray<T> &Primvar)
-			{
-				size_t SizePrimvar = Primvar.size();
-				CustomPrimvar.Data.reserve(Primvar.size());
-				for (const T& Value : Primvar)
-				{
-					if constexpr (std::is_arithmetic_v<std::remove_cvref_t<T>>)
-					{
-						CustomPrimvar.Data.emplace_back(Value);
-					}
-					else
-					{
-						pxr::GfVec4f Vec(1.f);
-						for (size_t IndexDimension = 0; IndexDimension < T::dimension; ++IndexDimension)
-						{
-							Vec.data()[IndexDimension] = Value[IndexDimension];
-						}
-						CustomPrimvar.Data.emplace_back(Vec);
-					}
-				}
-				return true;
-			};
-
-			bool bDataCollected = false;
+			FUsdMeshData::FPrimvar CustomPrimvar;
+			
 			if (pxr::VtArray<float> PrimvarFloat; PrimvarType == pxr::SdfValueTypeNames->FloatArray && Primvar.Get(&(PrimvarFloat), Options.TimeCode))
 			{
-				bDataCollected = CollectPrimvarData(PrimvarFloat);
+				CustomPrimvar.Data.Set<pxr::VtArray<float>>(PrimvarFloat);
 			}
 			else if (pxr::VtArray<int32> PrimvarInt; PrimvarType == pxr::SdfValueTypeNames->IntArray && Primvar.Get(&(PrimvarInt), Options.TimeCode))
 			{
-				bDataCollected = CollectPrimvarData(PrimvarInt);
+				CustomPrimvar.Data.Set<pxr::VtArray<int32>>(PrimvarInt);
 			}
 			else if (pxr::VtArray<bool> PrimvarBool; PrimvarType == pxr::SdfValueTypeNames->BoolArray && Primvar.Get(&(PrimvarBool), Options.TimeCode))
 			{
-				bDataCollected = CollectPrimvarData(PrimvarBool);
+				CustomPrimvar.Data.Set<pxr::VtArray<bool>>(PrimvarBool);
 			}
 			else if (pxr::VtArray<pxr::GfVec2f> PrimvarFloat2; PrimvarType == pxr::SdfValueTypeNames->Float2Array && Primvar.Get(&(PrimvarFloat2), Options.TimeCode))
 			{
-				bDataCollected = CollectPrimvarData(PrimvarFloat2);
+				CustomPrimvar.Data.Set<pxr::VtArray<pxr::GfVec2f>>(PrimvarFloat2);
 			}
 			else if (pxr::VtArray<pxr::GfVec3f> PrimvarFloat3; PrimvarType == pxr::SdfValueTypeNames->Float3Array && Primvar.Get(&(PrimvarFloat3), Options.TimeCode))
 			{
-				bDataCollected = CollectPrimvarData(PrimvarFloat3);
+				CustomPrimvar.Data.Set<pxr::VtArray<pxr::GfVec3f>>(PrimvarFloat3);
 			}
 			else if (pxr::VtArray<pxr::GfVec4f> PrimvarFloat4; PrimvarType == pxr::SdfValueTypeNames->Float4Array && Primvar.Get(&(PrimvarFloat4), Options.TimeCode))
 			{
-				CustomPrimvar.Data = MoveTemp(PrimvarFloat4);
-				bDataCollected = true;
+				CustomPrimvar.Data.Set<pxr::VtArray<pxr::GfVec4f>>(PrimvarFloat4);
 			}
+			else
+			{
+				UE_LOG(LogUsd, Warning, TEXT("Storage of primvar type '%s' not yet supported."), ANSI_TO_TCHAR(PrimvarType.c_str()));
+			}
+
+			bool bDataCollected = false;
+			Visit([&bDataCollected](auto&& Value)
+				  {
+					  bDataCollected = !Value.empty();
+				  }, CustomPrimvar.Data);
 
 			// only add if we collected the data
 			if(bDataCollected)
 			{
 				CustomPrimvar.Interpolation = GetAttrInterpolation(Primvar, DefaultInterpolation);
 				Primvar.GetIndices(&CustomPrimvar.Indices, Options.TimeCode);
-				InOutMeshData.CustomPrimvars.FindOrAdd(PrimvarName.c_str(), MoveTemp(CustomPrimvar));
+				InOutMeshData.Primvars.FindOrAdd(PrimvarName.c_str(), MoveTemp(CustomPrimvar));
 			}
 		}
 	}
@@ -2487,26 +2488,42 @@ namespace UE::UsdGeomMeshConversion::Private
 			UVIndices = {};
 		}
 
-		for (TPair<FString, FUsdMeshData::FCustomPrimvar>& Pair : InOutMeshData.CustomPrimvars)
+		for (TPair<FString, FUsdMeshData::FPrimvar>& Pair : InOutMeshData.Primvars)
 		{
-			FUsdMeshData::FCustomPrimvar& Primvar = Pair.Value;
-			Primvar.Data = ComputeFlattened(Primvar.Data, Primvar.Indices);
+			FUsdMeshData::FPrimvar& Primvar = Pair.Value;
+
+			Visit( TOverloaded{
+				[&Primvar](pxr::VtArray<bool>& Data){Data = ComputeFlattened(Data, Primvar.Indices);},
+				[&Primvar](pxr::VtArray<int32>& Data){Data = ComputeFlattened(Data, Primvar.Indices);},
+				[&Primvar](pxr::VtArray<float>& Data){Data = ComputeFlattened(Data, Primvar.Indices);},
+				[&Primvar](pxr::VtArray<pxr::GfVec2f>& Data){Data = ComputeFlattened(Data, Primvar.Indices);},
+				[&Primvar](pxr::VtArray<pxr::GfVec3f>& Data){Data = ComputeFlattened(Data, Primvar.Indices);},
+				[&Primvar](pxr::VtArray<pxr::GfVec4f>& Data){Data = ComputeFlattened(Data, Primvar.Indices);}
+			}, Primvar.Data);
+
 			Primvar.Indices = {};
 		}
 	}
 
-	void RegisterPrimvarAttributes(const TMap<FString, FUsdMeshData::FCustomPrimvar>& InCustomPrimvars, FMeshDescription& OutMeshDescription)
+	void RegisterPrimvarAttributes(const TMap<FString, FUsdMeshData::FPrimvar>& InCustomPrimvars, FMeshDescription& OutMeshDescription)
 	{
 		// Register Primvars attribute
-		for (const TPair<FString, FUsdMeshData::FCustomPrimvar>& Pair : InCustomPrimvars)
+		for (const TPair<FString, FUsdMeshData::FPrimvar>& Pair : InCustomPrimvars)
 		{
 			const FString& PrimvarName = Pair.Key;
-			const FUsdMeshData::FCustomPrimvar& CustomPrimvar = Pair.Value;
+			const FUsdMeshData::FPrimvar& CustomPrimvar = Pair.Value;
 
 			// Register the primvar if it doesn't exist
 			if (!OutMeshDescription.VertexInstanceAttributes().HasAttribute(*PrimvarName))
 			{
-				OutMeshDescription.VertexInstanceAttributes().RegisterAttribute<FVector4f>(*PrimvarName);
+				Visit(TOverloaded{
+					[&PrimvarName, &OutMeshDescription](const pxr::VtArray<bool>& ){OutMeshDescription.VertexInstanceAttributes().RegisterAttribute<bool>(*PrimvarName);},
+					[&PrimvarName, &OutMeshDescription](const pxr::VtArray<int32>& ){OutMeshDescription.VertexInstanceAttributes().RegisterAttribute<int32>(*PrimvarName);},
+					[&PrimvarName, &OutMeshDescription](const pxr::VtArray<float>& ){OutMeshDescription.VertexInstanceAttributes().RegisterAttribute<float>(*PrimvarName);},
+					[&PrimvarName, &OutMeshDescription](const pxr::VtArray<pxr::GfVec2f>& ){OutMeshDescription.VertexInstanceAttributes().RegisterAttribute<FVector2f>(*PrimvarName);},
+					[&PrimvarName, &OutMeshDescription](const pxr::VtArray<pxr::GfVec3f>& ){OutMeshDescription.VertexInstanceAttributes().RegisterAttribute<FVector3f>(*PrimvarName);},
+					[&PrimvarName, &OutMeshDescription](const pxr::VtArray<pxr::GfVec4f>& ){OutMeshDescription.VertexInstanceAttributes().RegisterAttribute<FVector4f>(*PrimvarName);}
+					}, CustomPrimvar.Data);
 			}
 		}
 	}
@@ -2532,7 +2549,7 @@ namespace UE::UsdGeomMeshConversion::Private
 		{
 			ensure(InMeshData.UVSetIndices[UVSetIndex].empty());
 		}
-		for (const TPair<FString, FUsdMeshData::FCustomPrimvar>& Pair : InMeshData.CustomPrimvars)
+		for (const TPair<FString, FUsdMeshData::FPrimvar>& Pair : InMeshData.Primvars)
 		{
 			ensure(Pair.Value.Indices.empty());
 		}
@@ -2703,7 +2720,7 @@ namespace UE::UsdGeomMeshConversion::Private
 			}
 
 			// Register Primvars attribute
-			RegisterPrimvarAttributes(InMeshData.CustomPrimvars, OutMeshDescription);
+			RegisterPrimvarAttributes(InMeshData.Primvars, OutMeshDescription);
 
 			// UVs
 			TVertexInstanceAttributesRef<FVector2f> MeshDescriptionUVs = StaticMeshAttributes.GetVertexInstanceUVs();
@@ -2775,12 +2792,38 @@ namespace UE::UsdGeomMeshConversion::Private
 																					.GetAttributesRef<FVector3f>(
 																						MeshAttribute::VertexInstance::Velocity
 																					);
+						
+			using FPrimvarVertexInstanceType = TVariant<
+				TVertexInstanceAttributesRef<bool>,
+				TVertexInstanceAttributesRef<int32>,
+				TVertexInstanceAttributesRef<float>,
+				TVertexInstanceAttributesRef<FVector2f>,
+				TVertexInstanceAttributesRef<FVector3f>,
+				TVertexInstanceAttributesRef<FVector4f>>;
+
+			TMap<FString, FPrimvarVertexInstanceType> MeshDescriptionPrimvar;
 
-			// We handle all the primvars as FVector4f, because during the baking phase we'll put them in the Color channel (see InterchangeGenericTexturePipeline::ExecutePostFactoryPipeline)
-			TMap<FString, TVertexInstanceAttributesRef<FVector4f>> MeshDescriptionPrimvar;
-			for (const TPair<FString, FUsdMeshData::FCustomPrimvar>& Pair : InMeshData.CustomPrimvars)
+			for (const TPair<FString, FUsdMeshData::FPrimvar>& Pair : InMeshData.Primvars)
 			{
-				MeshDescriptionPrimvar.Add(Pair.Key, OutMeshDescription.VertexInstanceAttributes().GetAttributesRef<FVector4f>(*(Pair.Key)));
+				Visit(TOverloaded{
+				[&MeshDescriptionPrimvar, &Pair, &OutMeshDescription](const pxr::VtArray<bool>&)
+				{MeshDescriptionPrimvar.Emplace(Pair.Key,FPrimvarVertexInstanceType{ TInPlaceType<TVertexInstanceAttributesRef<bool>>{}, OutMeshDescription.VertexInstanceAttributes().GetAttributesRef<bool>(*(Pair.Key))});},
+
+				[&MeshDescriptionPrimvar, &Pair, &OutMeshDescription](const pxr::VtArray<int32>&)
+				{MeshDescriptionPrimvar.Emplace(Pair.Key,FPrimvarVertexInstanceType{ TInPlaceType<TVertexInstanceAttributesRef<int32>>{}, OutMeshDescription.VertexInstanceAttributes().GetAttributesRef<int32>(*(Pair.Key))});},
+
+				[&MeshDescriptionPrimvar, &Pair, &OutMeshDescription](const pxr::VtArray<float>&)
+				{MeshDescriptionPrimvar.Emplace(Pair.Key,FPrimvarVertexInstanceType{ TInPlaceType<TVertexInstanceAttributesRef<float>>{}, OutMeshDescription.VertexInstanceAttributes().GetAttributesRef<float>(*(Pair.Key))});},
+
+				[&MeshDescriptionPrimvar, &Pair, &OutMeshDescription](const pxr::VtArray<pxr::GfVec2f>&)
+				{MeshDescriptionPrimvar.Emplace(Pair.Key,FPrimvarVertexInstanceType{ TInPlaceType<TVertexInstanceAttributesRef<FVector2f>>{}, OutMeshDescription.VertexInstanceAttributes().GetAttributesRef<FVector2f>(*(Pair.Key))});},
+
+				[&MeshDescriptionPrimvar, &Pair, &OutMeshDescription](const pxr::VtArray<pxr::GfVec3f>&)
+				{MeshDescriptionPrimvar.Emplace(Pair.Key,FPrimvarVertexInstanceType{ TInPlaceType<TVertexInstanceAttributesRef<FVector3f>>{}, OutMeshDescription.VertexInstanceAttributes().GetAttributesRef<FVector3f>(*(Pair.Key))});},
+
+				[&MeshDescriptionPrimvar, &Pair, &OutMeshDescription](const pxr::VtArray<pxr::GfVec4f>&)
+				{MeshDescriptionPrimvar.Emplace(Pair.Key,FPrimvarVertexInstanceType{ TInPlaceType<TVertexInstanceAttributesRef<FVector4f>>{}, OutMeshDescription.VertexInstanceAttributes().GetAttributesRef<FVector4f>(*(Pair.Key))});}
+				}, Pair.Value.Data);
 			}
 
 			OutMeshDescription.ReserveNewVertexInstances(InMeshData.FaceVertexCounts.size() * 3);
@@ -2955,37 +2998,64 @@ namespace UE::UsdGeomMeshConversion::Private
 					}
 
 					//Primvars
-					for (const TPair<FString, FUsdMeshData::FCustomPrimvar>& Pair : InMeshData.CustomPrimvars)
+					for (const TPair<FString, FUsdMeshData::FPrimvar>& Pair : InMeshData.Primvars)
 					{
 						const FString& PrimvarName = Pair.Key;
-						const FUsdMeshData::FCustomPrimvar& CustomPrimvar = Pair.Value;
+						const FUsdMeshData::FPrimvar& CustomPrimvar = Pair.Value;
 						const size_t ValueIndex = GetPrimValueIndex(CustomPrimvar.Interpolation,
 																	ControlPointIndex,
 																	CurrentVertexInstanceIndex,
 																	PolygonIndex);
 
-						if (!CustomPrimvar.Data.empty())
-						{
-							if (CustomPrimvar.Data.size() > ValueIndex)
-							{
-								const pxr::GfVec4f& VertexPrimvar = CustomPrimvar.Data[ValueIndex];
-								TVertexInstanceAttributesRef<FVector4f>& VertexInstanceAttributesPrimvar = *MeshDescriptionPrimvar.Find(PrimvarName);
-								VertexInstanceAttributesPrimvar[AddedVertexInstanceId] = FVector4f{ VertexPrimvar[0], VertexPrimvar[1], VertexPrimvar[2], VertexPrimvar[3] };
-							}
-							else
-							{
-								UE_LOG(
-									LogUsd,
-									Warning,
-									TEXT("Trying to read Primvar '%s' at index %u from prim '%s' but the prim only has %u values! Using "
-									"zeros instead."),
-									*PrimvarName,
-									ValueIndex,
-									*InMeshData.SourcePrimPath,
-									InMeshData.DisplayColors.size()
-								);
-							}
-						}
+						// We visit InMeshData.Primvar.Data for each type and make sure to fill the correct TVertexInstanceAttributesRef with the right type of the primvar
+						Visit([ValueIndex, AddedVertexInstanceId, &PrimvarName, &MeshDescriptionPrimvar, &InMeshData](auto && Value)
+							  {
+								  if (!Value.empty())
+								  {
+									  if (Value.size() > ValueIndex)
+									  {
+										  // underlying type of the pxr::VtArray that the primvar holds
+										  using FPrimvarType = std::decay_t<decltype(Value)>::ElementType;
+
+										  FPrimvarVertexInstanceType& PrimvarVertexInstance = MeshDescriptionPrimvar.FindChecked(PrimvarName);
+
+										  if constexpr (std::is_same_v<FPrimvarType, pxr::GfVec2f>)
+										  {
+											  TVertexInstanceAttributesRef<FVector2f>& VertexInstance = PrimvarVertexInstance.Get<TVertexInstanceAttributesRef<FVector2f>>();
+											  VertexInstance[AddedVertexInstanceId] = FVector2f{ Value[ValueIndex][0], Value[ValueIndex][1] };
+										  }
+										  else if constexpr (std::is_same_v<FPrimvarType, pxr::GfVec3f>)
+										  {
+											  TVertexInstanceAttributesRef<FVector3f>& VertexInstance = PrimvarVertexInstance.Get<TVertexInstanceAttributesRef<FVector3f>>();
+											  VertexInstance[AddedVertexInstanceId] = FVector3f{ Value[ValueIndex][0], Value[ValueIndex][1], Value[ValueIndex][2] };
+										  }
+										  else if constexpr (std::is_same_v<FPrimvarType, pxr::GfVec4f>)
+										  {
+											  TVertexInstanceAttributesRef<FVector4f>& VertexInstance = PrimvarVertexInstance.Get<TVertexInstanceAttributesRef<FVector4f>>();
+											  VertexInstance[AddedVertexInstanceId] = FVector4f{ Value[ValueIndex][0], Value[ValueIndex][1], Value[ValueIndex][2], Value[ValueIndex][3] };
+
+										  }
+										  else if constexpr (std::is_arithmetic_v<FPrimvarType>)
+										  {
+											  TVertexInstanceAttributesRef<FPrimvarType>& VertexInstance = PrimvarVertexInstance.Get<TVertexInstanceAttributesRef<FPrimvarType>>();
+											  VertexInstance[AddedVertexInstanceId] = Value[ValueIndex];
+										  }
+									  }
+									  else
+									  {
+										  UE_LOG(
+											  LogUsd,
+											  Warning,
+											  TEXT("Trying to read Primvar '%s' at index %u from prim '%s' but the prim only has %u values! Using "
+											  "zeros instead."),
+											  *PrimvarName,
+											  ValueIndex,
+											  *InMeshData.SourcePrimPath,
+											  InMeshData.DisplayColors.size()
+										  );
+									  }
+								  }
+							  }, CustomPrimvar.Data);
 					}
 				}
 
