Interchange: Allow specifying collision shapes for mesh nodes explicitly;
USD: Add support for physics collision schemas via USD Interchange.

diff --git a/Engine/Plugins/Importers/USDImporter/Source/USDSchemas/Private/USDGeomMeshTranslator.cpp b/Engine/Plugins/Importers/USDImporter/Source/USDSchemas/Private/USDGeomMeshTranslator.cpp
--- a/Engine/Plugins/Importers/USDImporter/Source/USDSchemas/Private/USDGeomMeshTranslator.cpp
+++ b/Engine/Plugins/Importers/USDImporter/Source/USDSchemas/Private/USDGeomMeshTranslator.cpp
@@ -1256,21 +1256,10 @@ namespace UE::UsdCollision::Private
 		CollisionShapes.ConvexElems.Add(ConvexElem);
 	}
 
-	// Approximation types from PhysicsMeshCollisionAPI and UE-specific approximations
-	enum class EUsdCollisionType : uint8
-	{
-		None,	 // no approximation so equivalent to CTF_UseComplexAsSimple
-		ConvexDecomposition,
-		ConvexHull,
-		Sphere,
-		Cube,
-		MeshSimplification,
-		Capsule,
-		CustomMesh
-	};
-
 	void CollectCustomCollisionShapes(const pxr::UsdPrim& UsdPrim, FKAggregateGeom& CollisionShapes)
 	{
+		using namespace UsdUtils;
+
 		UsdToUnreal::FUsdMeshConversionOptions Options;
 		Options.PurposesToLoad = EUsdPurpose::Guide;	// custom collision mesh must have guide purpose
 		Options.bMergeIdenticalMaterialSlots = false;
@@ -1298,7 +1287,7 @@ namespace UE::UsdCollision::Private
 				if (UsdToUnreal::ConvertGeomMesh(UsdMesh, MeshDescription, MaterialAssignments, Options))
 				{
 					// By default, use the custom mesh as convex hull, which takes care of the UCX_ prefix
-					EUsdCollisionType Approximation(EUsdCollisionType::ConvexHull);
+					EUsdCollisionType Approximation{EUsdCollisionType::ConvexHull};
 
 					// Use the FBX custom collision name convention as a hint
 					FString PrimName(UsdToUnreal::ConvertToken(ChildPrim.GetName()));
@@ -1345,6 +1334,8 @@ namespace UE::UsdCollision::Private
 
 	void SetupSimpleCollision(const pxr::UsdPrim& UsdPrim, UStaticMesh& StaticMesh)
 	{
+		using namespace UsdUtils;
+
 		UBodySetup* BodySetup = StaticMesh.GetBodySetup();
 		if (!BodySetup)
 		{
@@ -1380,73 +1371,6 @@ namespace UE::UsdCollision::Private
 			return;
 		}
 
-		EUsdCollisionType Approximation(EUsdCollisionType::None);
-		if (pxr::UsdGeomMesh UsdMesh{UsdPrim})
-		{
-			// Get the collision approximation type (only meshes should have UsdPhysicsMeshCollisionAPI)
-			if (pxr::UsdPhysicsMeshCollisionAPI MeshCollisionAPI{UsdMesh})
-			{
-				pxr::TfToken ApproximationAttr(pxr::UsdPhysicsTokens->none);
-				if (MeshCollisionAPI.GetApproximationAttr())
-				{
-					MeshCollisionAPI.GetApproximationAttr().Get(&ApproximationAttr);
-				}
-
-				if (ApproximationAttr == pxr::UsdPhysicsTokens->convexDecomposition)
-				{
-#if WITH_EDITOR
-					Approximation = EUsdCollisionType::ConvexDecomposition;
-#else
-					Approximation = EUsdCollisionType::ConvexHull;
-#endif
-				}
-				else if (ApproximationAttr == pxr::UsdPhysicsTokens->convexHull)
-				{
-					Approximation = EUsdCollisionType::ConvexHull;
-				}
-				else if (ApproximationAttr == pxr::UsdPhysicsTokens->boundingSphere)
-				{
-					Approximation = EUsdCollisionType::Sphere;
-				}
-				else if (ApproximationAttr == pxr::UsdPhysicsTokens->boundingCube)
-				{
-					Approximation = EUsdCollisionType::Cube;
-				}
-				else if (ApproximationAttr == pxr::UsdPhysicsTokens->meshSimplification)
-				{
-					Approximation = EUsdCollisionType::MeshSimplification;
-				}
-			}
-		}
-		else
-		{
-			// Collision for primitives are converted to their closest approximation
-			if (UsdPrim.IsA(pxr::UsdGeomTokens->Capsule))
-			{
-				Approximation = EUsdCollisionType::Capsule;
-			}
-			else if (UsdPrim.IsA(pxr::UsdGeomTokens->Cone))
-			{
-				Approximation = EUsdCollisionType::CustomMesh;
-			}
-			else if (UsdPrim.IsA(pxr::UsdGeomTokens->Cube))
-			{
-				Approximation = EUsdCollisionType::Cube;
-			}
-			else if (UsdPrim.IsA(pxr::UsdGeomTokens->Cylinder))
-			{
-				Approximation = EUsdCollisionType::CustomMesh;
-			}
-			else if (UsdPrim.IsA(pxr::UsdGeomTokens->Sphere))
-			{
-				Approximation = EUsdCollisionType::Sphere;
-			}
-			else if (UsdPrim.IsA(pxr::UsdGeomTokens->Plane))
-			{
-				Approximation = EUsdCollisionType::CustomMesh;
-			}
-		}
-
 		BodySetup->Modify();
 		BodySetup->RemoveSimpleCollision();
 		BodySetup->bNeverNeedsCookedCollisionData = false;
@@ -1454,6 +1378,7 @@ namespace UE::UsdCollision::Private
 		BodySetup->DefaultInstance.SetCollisionProfileName(UCollisionProfile::BlockAllDynamic_ProfileName);
 		BodySetup->CollisionTraceFlag = CTF_UseDefault;
 
+		EUsdCollisionType Approximation = GetCollisionApproximationType(UsdPrim);
 		switch (Approximation)
 		{
 			case EUsdCollisionType::None:
diff --git a/Engine/Plugins/Interchange/Extensions/OpenUSD/Source/Import/Private/InterchangeUsdTranslator.cpp b/Engine/Plugins/Interchange/Extensions/OpenUSD/Source/Import/Private/InterchangeUsdTranslator.cpp
--- a/Engine/Plugins/Interchange/Extensions/OpenUSD/Source/Import/Private/InterchangeUsdTranslator.cpp
+++ b/Engine/Plugins/Interchange/Extensions/OpenUSD/Source/Import/Private/InterchangeUsdTranslator.cpp
@@ -1913,6 +1913,48 @@ namespace UE::InterchangeUsdTranslator::Private
 		}
 	}
 
+	EInterchangeMeshCollision ConvertApproximationType(UsdUtils::EUsdCollisionType Approximation)
+	{
+		// References:
+		// - InterchangeGenericStaticMeshPipeline.cpp, GetCollisionMeshType()
+		// - InterchangeGenericStaticMeshPipeline.cpp, AddLodDataToStaticMesh()
+
+		using namespace UsdUtils;
+
+		switch (Approximation)
+		{
+			// EInterchangeMeshCollision::None means no collision, so treat EUsdCollisionType::None
+			// as convex collision instead
+			case UsdUtils::EUsdCollisionType::None:
+			case UsdUtils::EUsdCollisionType::ConvexDecomposition:
+			case UsdUtils::EUsdCollisionType::ConvexHull:
+			case UsdUtils::EUsdCollisionType::MeshSimplification:
+			case UsdUtils::EUsdCollisionType::CustomMesh:
+			{
+				return EInterchangeMeshCollision::Convex18DOP;
+			}
+			case UsdUtils::EUsdCollisionType::Sphere:
+			{
+				return EInterchangeMeshCollision::Sphere;
+			}
+			case UsdUtils::EUsdCollisionType::Cube:
+			{
+				return EInterchangeMeshCollision::Box;
+			}
+			case UsdUtils::EUsdCollisionType::Capsule:
+			{
+				return EInterchangeMeshCollision::Capsule;
+			}
+			default:
+			{
+				ensure(false);
+				break;
+			}
+		}
+
+		return EInterchangeMeshCollision::None;
+	}
+
 	FString UInterchangeUSDTranslatorImpl::AddMeshNode(
 		const UE::FUsdPrim& Prim,
 		UInterchangeUsdTranslatorSettings* TranslatorSettings,
@@ -2011,6 +2053,69 @@ namespace UE::InterchangeUsdTranslator::Private
 		else
 		{
 			MeshNode->SetPayLoadKey(PayloadKey, EInterchangeMeshPayLoadType::STATIC);
+
+			if (UsdUtils::IsCollisionEnabledForPrim(Prim))
+			{
+				// If the mesh prim is flagged for collision schemas AND also setup to be an FBX-style collision mesh, prefer
+				// the FBX style. We do this here because if we setup both styles at the same time, then GetCollisionMeshType() from
+				// InterchangeGenericStaticMeshPipeline.cpp would prefer the explicit collisions described on the translated node.
+				// That seems like the right thing in general, but is not what we want for USD due to compatibility with the legacy
+				// USD Importer code
+				bool bSetCustomCollisionType = true;
+				{
+					const static TSet<FString> CollisionPrefixes = {TEXT("UBX"), TEXT("UCX"), TEXT("MCDCX"), TEXT("USP"), TEXT("UCP")};
+
+					// Check if this mesh is an FBX-style collider
+					{
+						int32 FirstUnderscoreIndex = INDEX_NONE;
+						if (NodeName.FindChar(TEXT('_'), FirstUnderscoreIndex) && FirstUnderscoreIndex != INDEX_NONE)
+						{
+							FString Prefix = NodeName.Left(FirstUnderscoreIndex);
+							if (CollisionPrefixes.Contains(Prefix))
+							{
+								bSetCustomCollisionType = false;
+							}
+						}
+					}
+
+					// Check if we have any siblings that are also FBX-style colliders pointing at this mesh prim.
+					// In that case we want to disable the collider for this mesh prim itself, so that it matches legacy USD behavior
+					if (bSetCustomCollisionType)
+					{
+						TArray<UE::FUsdPrim> Siblings = Prim.GetParent().GetChildren();
+						for (const UE::FUsdPrim& Sibling : Siblings)
+						{
+							if (Sibling == Prim || !UsdUtils::IsCollisionMesh(Sibling))
+							{
+								continue;
+							}
+
+							FString SiblingName = Sibling.GetName().ToString();
+
+							int32 FirstUnderscoreIndex = INDEX_NONE;
+							if (SiblingName.FindChar(TEXT('_'), FirstUnderscoreIndex) && FirstUnderscoreIndex != INDEX_NONE)
+							{
+								FString PotentialNodeName = SiblingName.RightChop(FirstUnderscoreIndex + 1);
+								if (PotentialNodeName.StartsWith(NodeName))
+								{
+									bSetCustomCollisionType = false;
+									break;
+								}
+							}
+						}
+					}
+				}
+
+				if (bSetCustomCollisionType)
+				{
+					UsdUtils::EUsdCollisionType Approximation = UsdUtils::GetCollisionApproximationType(Prim);
+					EInterchangeMeshCollision InterchangeApproximation = ConvertApproximationType(Approximation);
+					if (InterchangeApproximation != EInterchangeMeshCollision::None)
+					{
+						MeshNode->SetCustomCollisionType(InterchangeApproximation);
+					}
+				}
+			}
 		}
 
 		// Material assignments
diff --git a/Engine/Plugins/Interchange/Runtime/Source/Nodes/Private/InterchangeMeshNode.cpp b/Engine/Plugins/Interchange/Runtime/Source/Nodes/Private/InterchangeMeshNode.cpp
--- a/Engine/Plugins/Interchange/Runtime/Source/Nodes/Private/InterchangeMeshNode.cpp
+++ b/Engine/Plugins/Interchange/Runtime/Source/Nodes/Private/InterchangeMeshNode.cpp
@@ -538,6 +538,16 @@ bool UInterchangeMeshNode::RemoveSlotMaterialDependencyUid(const FString& SlotNa
 	return SlotMaterialDependencies.RemoveKey(SlotName);
 }
 
+bool UInterchangeMeshNode::GetCustomCollisionType(EInterchangeMeshCollision& AttributeValue) const
+{
+	IMPLEMENT_NODE_ATTRIBUTE_GETTER(CollisionType, EInterchangeMeshCollision);
+}
+
+bool UInterchangeMeshNode::SetCustomCollisionType(EInterchangeMeshCollision AttributeValue)
+{
+	IMPLEMENT_NODE_ATTRIBUTE_SETTER_NODELEGATE(CollisionType, EInterchangeMeshCollision);
+}
+
 UInterchangeGeometryCacheNode::UInterchangeGeometryCacheNode()
 {
 }
diff --git a/Engine/Plugins/Interchange/Runtime/Source/Nodes/Public/InterchangeMeshNode.h b/Engine/Plugins/Interchange/Runtime/Source/Nodes/Public/InterchangeMeshNode.h
--- a/Engine/Plugins/Interchange/Runtime/Source/Nodes/Public/InterchangeMeshNode.h
+++ b/Engine/Plugins/Interchange/Runtime/Source/Nodes/Public/InterchangeMeshNode.h
@@ -3,6 +3,7 @@
 #pragma once
 
 #include "CoreMinimal.h"
+#include "InterchangeMeshDefinitions.h"
 #include "Nodes/InterchangeBaseNode.h"
 
 #include "InterchangeMeshNode.generated.h"
@@ -115,7 +116,7 @@ public:
 	 * Icon names are created by adding "InterchangeIcon_" in front of the specialized type. If there is no special type, the function will return NAME_None, which will use the default icon.
 	 */
 	virtual FName GetIconName() const override;
-	
+
 	/**
 	 * Override Serialize() to restore SlotMaterialDependencies on load.
 	 */
@@ -170,7 +171,7 @@ public:
 
 	UFUNCTION(BlueprintCallable, Category = "Interchange | Node | Mesh")
 	virtual void SetPayLoadKey(const FString& PayLoadKey, const EInterchangeMeshPayLoadType& PayLoadType);
-	
+
 	/** Query the vertex count of this mesh. Return false if the attribute was not set.*/
 	UFUNCTION(BlueprintCallable, Category = "Interchange | Node | Mesh")
 	bool GetCustomVertexCount(int32& AttributeValue) const;
@@ -357,22 +358,35 @@ public:
 	UFUNCTION(BlueprintCallable, Category = "Interchange | Node | StaticMesh")
 	bool RemoveSlotMaterialDependencyUid(const FString& SlotName);
 
+	/**
+	 * Get the type of collision shapes we should generate from this mesh.
+	 * Note: This is a separate mechanism from the FBX-style collision shape name prefixes. For now, these
+	 * collision shapes will only be used for the static mesh generated from this very same Mesh node
+	 */
+	UFUNCTION(BlueprintCallable, Category = "Interchange | Node | Mesh")
+	bool GetCustomCollisionType(EInterchangeMeshCollision& AttributeValue) const;
+
+	/** Set the type of collision shapes we should generate from this mesh */
+	UFUNCTION(BlueprintCallable, Category = "Interchange | Node | Mesh")
+	bool SetCustomCollisionType(EInterchangeMeshCollision AttributeValue);
+
 private:
-	const UE::Interchange::FAttributeKey Macro_CustomVertexCountKey = UE::Interchange::FAttributeKey(TEXT("VertexCount"));
-	const UE::Interchange::FAttributeKey Macro_CustomPolygonCountKey = UE::Interchange::FAttributeKey(TEXT("PolygonCount"));
-	const UE::Interchange::FAttributeKey Macro_CustomBoundingBoxKey = UE::Interchange::FAttributeKey(TEXT("BoundingBox"));
-	const UE::Interchange::FAttributeKey Macro_CustomHasVertexNormalKey = UE::Interchange::FAttributeKey(TEXT("HasVertexNormal"));
-	const UE::Interchange::FAttributeKey Macro_CustomHasVertexBinormalKey = UE::Interchange::FAttributeKey(TEXT("HasVertexBinormal"));
-	const UE::Interchange::FAttributeKey Macro_CustomHasVertexTangentKey = UE::Interchange::FAttributeKey(TEXT("HasVertexTangent"));
-	const UE::Interchange::FAttributeKey Macro_CustomHasSmoothGroupKey = UE::Interchange::FAttributeKey(TEXT("HasSmoothGroup"));
-	const UE::Interchange::FAttributeKey Macro_CustomHasVertexColorKey = UE::Interchange::FAttributeKey(TEXT("HasVertexColor"));
-	const UE::Interchange::FAttributeKey Macro_CustomUVCountKey = UE::Interchange::FAttributeKey(TEXT("UVCount"));
+	IMPLEMENT_NODE_ATTRIBUTE_KEY(VertexCount);
+	IMPLEMENT_NODE_ATTRIBUTE_KEY(PolygonCount);
+	IMPLEMENT_NODE_ATTRIBUTE_KEY(BoundingBox);
+	IMPLEMENT_NODE_ATTRIBUTE_KEY(HasVertexNormal);
+	IMPLEMENT_NODE_ATTRIBUTE_KEY(HasVertexBinormal);
+	IMPLEMENT_NODE_ATTRIBUTE_KEY(HasVertexTangent);
+	IMPLEMENT_NODE_ATTRIBUTE_KEY(HasSmoothGroup);
+	IMPLEMENT_NODE_ATTRIBUTE_KEY(HasVertexColor);
+	IMPLEMENT_NODE_ATTRIBUTE_KEY(UVCount);
+	IMPLEMENT_NODE_ATTRIBUTE_KEY(CollisionType);
 
 	UE::Interchange::TArrayAttributeHelper<FString> SkeletonDependencies;
 	UE::Interchange::TArrayAttributeHelper<FString> MaterialDependencies;
 	UE::Interchange::TArrayAttributeHelper<FString> MorphTargetDependencies;
 	UE::Interchange::TArrayAttributeHelper<FString> SceneInstancesUids;
-	
+
 	UE::Interchange::TMapAttributeHelper<FString, FString> SlotMaterialDependencies;
 };
 
@@ -421,4 +435,4 @@ private:
 	IMPLEMENT_NODE_ATTRIBUTE_KEY(EndFrame);
 	IMPLEMENT_NODE_ATTRIBUTE_KEY(FrameRate);
 	IMPLEMENT_NODE_ATTRIBUTE_KEY(HasConstantTopology);
-};
\ No newline at end of file
+};
diff --git a/Engine/Plugins/Interchange/Runtime/Source/Pipelines/Private/InterchangeGenericStaticMeshPipeline.cpp b/Engine/Plugins/Interchange/Runtime/Source/Pipelines/Private/InterchangeGenericStaticMeshPipeline.cpp
--- a/Engine/Plugins/Interchange/Runtime/Source/Pipelines/Private/InterchangeGenericStaticMeshPipeline.cpp
+++ b/Engine/Plugins/Interchange/Runtime/Source/Pipelines/Private/InterchangeGenericStaticMeshPipeline.cpp
@@ -1,4 +1,4 @@
-// Copyright Epic Games, Inc. All Rights Reserved. 
+// Copyright Epic Games, Inc. All Rights Reserved.
 #include "InterchangeGenericMeshPipeline.h"
 
 #include "AssetRegistry/AssetRegistryModule.h"
@@ -63,88 +63,125 @@ namespace UE::Interchange::Private
 	}
 }
 
-static TTuple<EInterchangeMeshCollision, FString> GetCollisionMeshType(TObjectPtr<UInterchangePipelineMeshesUtilities> PipelineMeshesUtilities, const UInterchangeBaseNodeContainer& NodeContainer, const FString& NodeUid, const TArray<FString>& AllNodeUids)
+static TTuple<EInterchangeMeshCollision, FString> GetCollisionMeshType(
+	TObjectPtr<UInterchangePipelineMeshesUtilities> PipelineMeshesUtilities,
+	const UInterchangeBaseNodeContainer& NodeContainer,
+	const FString& NodeUid,
+	const TArray<FString>& AllNodeUids,
+	bool bImportCollisionAccordingToMeshName
+)
 {
-	FString MeshName = UE::Interchange::Private::GetNodeName(PipelineMeshesUtilities, NodeContainer, NodeUid);
 	EInterchangeMeshCollision CollisionType = EInterchangeMeshCollision::None;
 
-	// Determine if the mesh name is a potential collision mesh
-
-	if (MeshName.StartsWith(TEXT("UBX_")))
-	{
-		CollisionType = EInterchangeMeshCollision::Box;
-	}
-	else if (MeshName.StartsWith(TEXT("UCX_")) || MeshName.StartsWith(TEXT("MCDCX_")))
+	// Try to find the mesh node we're actually talking about
+	const UInterchangeBaseNode* BaseNode = NodeContainer.GetNode(NodeUid);
+	const UInterchangeMeshNode* MeshNode = Cast<const UInterchangeMeshNode>(BaseNode);
+	if (!MeshNode)
 	{
-		CollisionType = EInterchangeMeshCollision::Convex18DOP;
-	}
-	else if (MeshName.StartsWith(TEXT("USP_")))
-	{
-		CollisionType = EInterchangeMeshCollision::Sphere;
+		if (const UInterchangeSceneNode* SceneNode = Cast<UInterchangeSceneNode>(BaseNode))
+		{
+			FString MeshDependency;
+			if (SceneNode->GetCustomAssetInstanceUid(MeshDependency))
+			{
+				MeshNode = Cast<const UInterchangeMeshNode>(NodeContainer.GetNode(MeshDependency));
+			}
+		}
 	}
-	else if (MeshName.StartsWith(TEXT("UCP_")))
+
+	// If we have an explicit collision type on the mesh, we know the mesh node is a collision mesh itself
+	if (MeshNode)
 	{
-		CollisionType = EInterchangeMeshCollision::Capsule;
+		if (MeshNode->GetCustomCollisionType(CollisionType) && CollisionType != EInterchangeMeshCollision::None)
+		{
+			return {CollisionType, NodeUid};
+		}
 	}
-	else
+
+	if (bImportCollisionAccordingToMeshName)
 	{
-		return { EInterchangeMeshCollision::None, FString() };
-	}
+		FString MeshName = UE::Interchange::Private::GetNodeName(PipelineMeshesUtilities, NodeContainer, NodeUid);
 
-	// We have a mesh name with a collision type suffix.
-	// However it should only be treated as a collision mesh if its body name corresponds to one of the other meshes.
-	// If we get here, we know there is at least one underscore, so we never expect either of these character searches to fail.
+		// Determine if the mesh name is a potential collision mesh
 
-	int32 FirstUnderscore = INDEX_NONE;
-	verify(MeshName.FindChar(TEXT('_'), FirstUnderscore));
+		if (MeshName.StartsWith(TEXT("UBX_")))
+		{
+			CollisionType = EInterchangeMeshCollision::Box;
+		}
+		else if (MeshName.StartsWith(TEXT("UCX_")) || MeshName.StartsWith(TEXT("MCDCX_")))
+		{
+			CollisionType = EInterchangeMeshCollision::Convex18DOP;
+		}
+		else if (MeshName.StartsWith(TEXT("USP_")))
+		{
+			CollisionType = EInterchangeMeshCollision::Sphere;
+		}
+		else if (MeshName.StartsWith(TEXT("UCP_")))
+		{
+			CollisionType = EInterchangeMeshCollision::Capsule;
+		}
+		else
+		{
+			return { EInterchangeMeshCollision::None, FString() };
+		}
 
-	int32 LastUnderscore = INDEX_NONE;
-	verify(MeshName.FindLastChar(TEXT('_'), LastUnderscore));
+		// We have a mesh name with a collision type suffix.
+		// However it should only be treated as a collision mesh if its body name corresponds to one of the other meshes.
+		// If we get here, we know there is at least one underscore, so we never expect either of these character searches to fail.
 
-	auto MatchPredicate = [&PipelineMeshesUtilities , &NodeContainer](FString Body)
-	{
-		// Generate a predicate to be used by the below Finds.
-		return [Body, &PipelineMeshesUtilities, &NodeContainer](const FString& ToCompare) { return Body == UE::Interchange::Private::GetNodeName(PipelineMeshesUtilities, NodeContainer, ToCompare); };
-	};
+		int32 FirstUnderscore = INDEX_NONE;
+		verify(MeshName.FindChar(TEXT('_'), FirstUnderscore));
 
-	// If we find a mesh named the same as the collision mesh (following the collision prefix), we have a match
-	// e.g. this will match 'UBX_House' with a mesh called 'House'
+		int32 LastUnderscore = INDEX_NONE;
+		verify(MeshName.FindLastChar(TEXT('_'), LastUnderscore));
 
-	const FString* CorrespondingMeshUid = AllNodeUids.FindByPredicate(MatchPredicate(MeshName.RightChop(FirstUnderscore + 1)));
-	if (CorrespondingMeshUid)
-	{
-		return { CollisionType, *CorrespondingMeshUid };
-	}
+		auto MatchPredicate = [&PipelineMeshesUtilities , &NodeContainer](FString Body)
+		{
+			// Generate a predicate to be used by the below Finds.
+			return [Body, &PipelineMeshesUtilities, &NodeContainer](const FString& ToCompare) { return Body == UE::Interchange::Private::GetNodeName(PipelineMeshesUtilities, NodeContainer, ToCompare); };
+		};
 
-	// Otherwise strip the final underscore suffix from the collision mesh name and look again
-	// e.g. this will match 'UBX_House_01' with a mesh called 'House'
+		// If we find a mesh named the same as the collision mesh (following the collision prefix), we have a match
+		// e.g. this will match 'UBX_House' with a mesh called 'House'
 
-	if (FirstUnderscore != LastUnderscore)
-	{
-		CorrespondingMeshUid = AllNodeUids.FindByPredicate(MatchPredicate(MeshName.Mid(FirstUnderscore + 1, LastUnderscore - FirstUnderscore - 1)));
+		const FString* CorrespondingMeshUid = AllNodeUids.FindByPredicate(MatchPredicate(MeshName.RightChop(FirstUnderscore + 1)));
 		if (CorrespondingMeshUid)
 		{
 			return { CollisionType, *CorrespondingMeshUid };
 		}
-	}
 
-	// Mesh had a collision type prefix, but no corresponding mesh, so don't treat it as a collision mesh
+		// Otherwise strip the final underscore suffix from the collision mesh name and look again
+		// e.g. this will match 'UBX_House_01' with a mesh called 'House'
+
+		if (FirstUnderscore != LastUnderscore)
+		{
+			// Even if the mesh has a prefix, only treat it as a collision mesh if it has a corresponding render mesh
+			CorrespondingMeshUid = AllNodeUids.FindByPredicate(MatchPredicate(MeshName.Mid(FirstUnderscore + 1, LastUnderscore - FirstUnderscore - 1)));
+			if (CorrespondingMeshUid)
+			{
+				return { CollisionType, *CorrespondingMeshUid };
+			}
+		}
+
+	}
 
 	return { EInterchangeMeshCollision::None, FString() };
 }
 
-
-static bool IsCollisionMeshUid(TObjectPtr<UInterchangePipelineMeshesUtilities> PipelineMeshesUtilities, const UInterchangeBaseNodeContainer& NodeContainer, const FString& MeshUid, const TArray<FString>& MeshUids)
+// Returns true if MeshUid describes a Mesh node that is purely the collision mesh of some other mesh
+static bool IsJustCollisionMesh(TObjectPtr<UInterchangePipelineMeshesUtilities> PipelineMeshesUtilities, const UInterchangeBaseNodeContainer& NodeContainer, const FString& MeshUid, const TArray<FString>& MeshUids, bool bImportCollisionAccordingToMeshName)
 {
-	return GetCollisionMeshType(PipelineMeshesUtilities, NodeContainer, MeshUid, MeshUids).Get<0>() != EInterchangeMeshCollision::None;
+	TTuple<EInterchangeMeshCollision, FString> Tuple = GetCollisionMeshType(PipelineMeshesUtilities, NodeContainer, MeshUid, MeshUids, bImportCollisionAccordingToMeshName);
+
+	// If MeshUid is a collision mesh but *of itself*, then it is both a collision and a render mesh, and we should return false
+	return Tuple.Key != EInterchangeMeshCollision::None && Tuple.Value != MeshUid;
 }
 
 
-static void BuildMeshToCollisionMeshMap(TObjectPtr<UInterchangePipelineMeshesUtilities> PipelineMeshesUtilities, const UInterchangeBaseNodeContainer& NodeContainer, const TArray<FString>& MeshUids, TMap<FString, TArray<FString>>& MeshToCollisionMeshMap)
+static void BuildMeshToCollisionMeshMap(TObjectPtr<UInterchangePipelineMeshesUtilities> PipelineMeshesUtilities, const UInterchangeBaseNodeContainer& NodeContainer, const TArray<FString>& MeshUids, TMap<FString, TArray<FString>>& MeshToCollisionMeshMap, bool bImportCollisionAccordingToMeshName)
 {
 	for (const FString& MeshUid : MeshUids)
 	{
-		TTuple<EInterchangeMeshCollision, FString> CollisionType = GetCollisionMeshType(PipelineMeshesUtilities, NodeContainer, MeshUid, MeshUids);
+		TTuple<EInterchangeMeshCollision, FString> CollisionType = GetCollisionMeshType(PipelineMeshesUtilities, NodeContainer, MeshUid, MeshUids, bImportCollisionAccordingToMeshName);
 		if (CollisionType.Get<0>() != EInterchangeMeshCollision::None)
 		{
 			MeshToCollisionMeshMap.FindOrAdd(FString(CollisionType.Get<1>())).Emplace(MeshUid);
@@ -274,21 +311,20 @@ void UInterchangeGenericMeshPipeline::ExecutePreImportPipelineStaticMesh()
 				{
 					constexpr int32 LodIndexZero = 0;
 					bool bFoundMeshes = false;
+
 					// Work out which meshes are collision meshes which correspond to another mesh
 					TMap<FString, TArray<FString>> MeshToCollisionMeshMap;
-					if (bImportCollisionAccordingToMeshName)
-					{
-						BuildMeshToCollisionMeshMap(PipelineMeshesUtilities, *BaseNodeContainer, MeshUids, MeshToCollisionMeshMap);
-					}
+					BuildMeshToCollisionMeshMap(PipelineMeshesUtilities, *BaseNodeContainer, MeshUids, MeshToCollisionMeshMap, bImportCollisionAccordingToMeshName);
+
 					// Now iterate through each mesh UID, creating a new factory for each one
 					for (const FString& MeshUid : MeshUids)
 					{
-						if (bImportCollisionAccordingToMeshName && IsCollisionMeshUid(PipelineMeshesUtilities, *BaseNodeContainer, MeshUid, MeshUids))
+						// If this is just a collision mesh, don't add a factory node; it will be added as part of another factory node
+						if (IsJustCollisionMesh(PipelineMeshesUtilities, *BaseNodeContainer, MeshUid, MeshUids, bImportCollisionAccordingToMeshName))
 						{
-							// If this is a collision mesh, don't add a factory; it will be added as part of another factory
 							continue;
 						}
-						
+
 						TMap<int32, TArray<FString>> MeshUidsPerLodIndex;
 						if (bInstancedMesh)
 						{
@@ -316,11 +352,17 @@ void UInterchangeGenericMeshPipeline::ExecutePreImportPipelineStaticMesh()
 
 						if (MeshUidsPerLodIndex.Num() > 0)
 						{
-							if (bCollision && bImportCollisionAccordingToMeshName)
+							if (bCollision)
 							{
 								if (const TArray<FString>* CorrespondingCollisionMeshes = MeshToCollisionMeshMap.Find(MeshUid))
 								{
-									MeshUidsPerLodIndex.FindOrAdd(LodIndexZero).Append(*CorrespondingCollisionMeshes);
+									TArray<FString>& TranslatedNodes = MeshUidsPerLodIndex.FindOrAdd(LodIndexZero);
+									for (const FString& CollisionMesh : *CorrespondingCollisionMeshes)
+									{
+										// AddUnique here because now we support meshes being both collision AND render meshes at the
+										// same time we may already have these entries in MeshUidsPerLodIndex
+										TranslatedNodes.AddUnique(CollisionMesh);
+									}
 								}
 							}
 
@@ -427,7 +469,7 @@ UInterchangeStaticMeshFactoryNode* UInterchangeGenericMeshPipeline::CreateStatic
 
 	StaticMeshFactoryNode->InitializeStaticMeshNode(StaticMeshUid, DisplayLabel, UStaticMesh::StaticClass()->GetName());
 	BaseNodeContainer->AddNode(StaticMeshFactoryNode);
-	
+
 	//Set the pipeline import sockets property
 	StaticMeshFactoryNode->SetCustomImportSockets(CommonMeshesProperties->bImportSockets);
 
@@ -468,7 +510,7 @@ UInterchangeStaticMeshFactoryNode* UInterchangeGenericMeshPipeline::CreateStatic
 	StaticMeshFactoryNode->SetCustomUseFullPrecisionUVs(CommonMeshesProperties->bUseFullPrecisionUVs);
 	StaticMeshFactoryNode->SetCustomUseBackwardsCompatibleF16TruncUVs(CommonMeshesProperties->bUseBackwardsCompatibleF16TruncUVs);
 	StaticMeshFactoryNode->SetCustomRemoveDegenerates(CommonMeshesProperties->bRemoveDegenerates);
-	
+
 	//Static meshes build options
 	StaticMeshFactoryNode->SetCustomBuildReversedIndexBuffer(bBuildReversedIndexBuffer);
 	StaticMeshFactoryNode->SetCustomGenerateLightmapUVs(bGenerateLightmapUVs);
@@ -484,7 +526,7 @@ UInterchangeStaticMeshFactoryNode* UInterchangeGenericMeshPipeline::CreateStatic
 	StaticMeshFactoryNode->SetCustomBuildNanite(bBuildNanite);
 	StaticMeshFactoryNode->SetCustomAutoComputeLODScreenSizes(bAutoComputeLODScreenSizes);
 	StaticMeshFactoryNode->SetLODScreenSizes(LODScreenSizes);
-	
+
 	return StaticMeshFactoryNode;
 }
 
@@ -577,7 +619,7 @@ void UInterchangeGenericMeshPipeline::AddLodDataToStaticMesh(UInterchangeStaticM
 				StaticMeshFactoryNode->AddTargetNodeUid(NodeUid);
 				StaticMeshFactoryNode->AddSocketUids(PipelineMeshesUtilities->GetMeshGeometryByUid(NodeUid).AttachedSocketUids);
 				MeshNode->AddTargetNodeUid(StaticMeshFactoryNode->GetUniqueID());
-				
+
 				MeshNode->GetSlotMaterialDependencies(SlotMaterialDependencies);
 			}
 
@@ -585,7 +627,8 @@ void UInterchangeGenericMeshPipeline::AddLodDataToStaticMesh(UInterchangeStaticM
 
 			if (bImportCollisionAccordingToMeshName)
 			{
-				TTuple<EInterchangeMeshCollision, FStringView> MeshType = GetCollisionMeshType(PipelineMeshesUtilities, *BaseNodeContainer, NodeUid, NodeUids);
+				TTuple<EInterchangeMeshCollision, FString> MeshType = GetCollisionMeshType(PipelineMeshesUtilities, *BaseNodeContainer, NodeUid, NodeUids, bImportCollisionAccordingToMeshName);
+
 				switch (MeshType.Get<0>())
 				{
 				case EInterchangeMeshCollision::None:
@@ -612,6 +655,13 @@ void UInterchangeGenericMeshPipeline::AddLodDataToStaticMesh(UInterchangeStaticM
 					LodDataNode->AddConvexCollisionMeshUid(NodeUid);
 					break;
 				}
+
+				// If the Mesh is both render AND collision mesh let's also add it as a render mesh,
+				// as the above switch will have only added the collision mesh to the LodDataNode
+				if (MeshType.Get<1>() == NodeUid && MeshType.Get<0>() != EInterchangeMeshCollision::None)
+				{
+					LodDataNode->AddMeshUid(NodeUid);
+				}
 			}
 			else
 			{
diff --git a/Engine/Plugins/Runtime/USDCore/Source/USDUtilities/Private/USDGeomMeshConversion.cpp b/Engine/Plugins/Runtime/USDCore/Source/USDUtilities/Private/USDGeomMeshConversion.cpp
--- a/Engine/Plugins/Runtime/USDCore/Source/USDUtilities/Private/USDGeomMeshConversion.cpp
+++ b/Engine/Plugins/Runtime/USDCore/Source/USDUtilities/Private/USDGeomMeshConversion.cpp
@@ -1211,8 +1211,11 @@ namespace UE::UsdGeomMeshConversion::Private
 	// https://en.cppreference.com/w/cpp/utility/variant/visit2
 	// Helper type for the visitor
 	template<class... Ts>
-	struct TOverloaded : Ts... {using Ts::operator()...; };
-	//Explicit deduction guide (not needed as of C++20
+	struct TOverloaded : Ts...
+	{
+		using Ts::operator()...;
+	};
+	// Explicit deduction guide (not needed as of C++20
 	template<class... Ts>
 	TOverloaded(Ts...) -> TOverloaded<Ts...>;
 
@@ -1299,14 +1302,11 @@ namespace UE::UsdGeomMeshConversion::Private
 		}
 	}	 // namespace USDSkeletalDataConversion
 
-	void CollectPrimvarData(
-		const pxr::UsdPrim& Prim,
-		const UsdToUnreal::FUsdMeshConversionOptions& Options,
-		FUsdMeshData& InOutMeshData)
+	void CollectPrimvarData(const pxr::UsdPrim& Prim, const UsdToUnreal::FUsdMeshConversionOptions& Options, FUsdMeshData& InOutMeshData)
 	{
 		const pxr::TfToken DefaultInterpolation = pxr::UsdGeomTokens->varying;
 
-		pxr::UsdGeomPrimvarsAPI PrimvarAPI{ Prim };
+		pxr::UsdGeomPrimvarsAPI PrimvarAPI{Prim};
 		std::vector<pxr::UsdGeomPrimvar> Primvars = PrimvarAPI.GetPrimvarsWithAuthoredValues();
 		for (const pxr::UsdGeomPrimvar& Primvar : Primvars)
 		{
@@ -1318,7 +1318,7 @@ namespace UE::UsdGeomMeshConversion::Private
 
 			std::string PrimvarType = Primvar.GetTypeName().GetAsToken().GetString();
 			FUsdMeshData::FPrimvar CustomPrimvar;
-			
+
 			if (pxr::VtArray<float> PrimvarFloat; PrimvarType == pxr::SdfValueTypeNames->FloatArray && Primvar.Get(&(PrimvarFloat), Options.TimeCode))
 			{
 				CustomPrimvar.Data.Set<pxr::VtArray<float>>(PrimvarFloat);
@@ -1327,19 +1327,23 @@ namespace UE::UsdGeomMeshConversion::Private
 			{
 				CustomPrimvar.Data.Set<pxr::VtArray<int32>>(PrimvarInt);
 			}
-			else if (pxr::VtArray<bool> PrimvarBool; PrimvarType == pxr::SdfValueTypeNames->BoolArray && Primvar.Get(&(PrimvarBool), Options.TimeCode))
+			else if (pxr::VtArray<bool> PrimvarBool;
+					 PrimvarType == pxr::SdfValueTypeNames->BoolArray && Primvar.Get(&(PrimvarBool), Options.TimeCode))
 			{
 				CustomPrimvar.Data.Set<pxr::VtArray<bool>>(PrimvarBool);
 			}
-			else if (pxr::VtArray<pxr::GfVec2f> PrimvarFloat2; PrimvarType == pxr::SdfValueTypeNames->Float2Array && Primvar.Get(&(PrimvarFloat2), Options.TimeCode))
+			else if (pxr::VtArray<pxr::GfVec2f> PrimvarFloat2;
+					 PrimvarType == pxr::SdfValueTypeNames->Float2Array && Primvar.Get(&(PrimvarFloat2), Options.TimeCode))
 			{
 				CustomPrimvar.Data.Set<pxr::VtArray<pxr::GfVec2f>>(PrimvarFloat2);
 			}
-			else if (pxr::VtArray<pxr::GfVec3f> PrimvarFloat3; PrimvarType == pxr::SdfValueTypeNames->Float3Array && Primvar.Get(&(PrimvarFloat3), Options.TimeCode))
+			else if (pxr::VtArray<pxr::GfVec3f> PrimvarFloat3;
+					 PrimvarType == pxr::SdfValueTypeNames->Float3Array && Primvar.Get(&(PrimvarFloat3), Options.TimeCode))
 			{
 				CustomPrimvar.Data.Set<pxr::VtArray<pxr::GfVec3f>>(PrimvarFloat3);
 			}
-			else if (pxr::VtArray<pxr::GfVec4f> PrimvarFloat4; PrimvarType == pxr::SdfValueTypeNames->Float4Array && Primvar.Get(&(PrimvarFloat4), Options.TimeCode))
+			else if (pxr::VtArray<pxr::GfVec4f> PrimvarFloat4;
+					 PrimvarType == pxr::SdfValueTypeNames->Float4Array && Primvar.Get(&(PrimvarFloat4), Options.TimeCode))
 			{
 				CustomPrimvar.Data.Set<pxr::VtArray<pxr::GfVec4f>>(PrimvarFloat4);
 			}
@@ -1349,13 +1353,16 @@ namespace UE::UsdGeomMeshConversion::Private
 			}
 
 			bool bDataCollected = false;
-			Visit([&bDataCollected](auto&& Value)
-				  {
-					  bDataCollected = !Value.empty();
-				  }, CustomPrimvar.Data);
+			Visit(
+				[&bDataCollected](auto&& Value)
+				{
+					bDataCollected = !Value.empty();
+				},
+				CustomPrimvar.Data
+			);
 
 			// only add if we collected the data
-			if(bDataCollected)
+			if (bDataCollected)
 			{
 				CustomPrimvar.Interpolation = GetAttrInterpolation(Primvar, DefaultInterpolation);
 				Primvar.GetIndices(&CustomPrimvar.Indices, Options.TimeCode);
@@ -2492,14 +2499,35 @@ namespace UE::UsdGeomMeshConversion::Private
 		{
 			FUsdMeshData::FPrimvar& Primvar = Pair.Value;
 
-			Visit( TOverloaded{
-				[&Primvar](pxr::VtArray<bool>& Data){Data = ComputeFlattened(Data, Primvar.Indices);},
-				[&Primvar](pxr::VtArray<int32>& Data){Data = ComputeFlattened(Data, Primvar.Indices);},
-				[&Primvar](pxr::VtArray<float>& Data){Data = ComputeFlattened(Data, Primvar.Indices);},
-				[&Primvar](pxr::VtArray<pxr::GfVec2f>& Data){Data = ComputeFlattened(Data, Primvar.Indices);},
-				[&Primvar](pxr::VtArray<pxr::GfVec3f>& Data){Data = ComputeFlattened(Data, Primvar.Indices);},
-				[&Primvar](pxr::VtArray<pxr::GfVec4f>& Data){Data = ComputeFlattened(Data, Primvar.Indices);}
-			}, Primvar.Data);
+			Visit(
+				TOverloaded{
+					[&Primvar](pxr::VtArray<bool>& Data)
+					{
+						Data = ComputeFlattened(Data, Primvar.Indices);
+					},
+					[&Primvar](pxr::VtArray<int32>& Data)
+					{
+						Data = ComputeFlattened(Data, Primvar.Indices);
+					},
+					[&Primvar](pxr::VtArray<float>& Data)
+					{
+						Data = ComputeFlattened(Data, Primvar.Indices);
+					},
+					[&Primvar](pxr::VtArray<pxr::GfVec2f>& Data)
+					{
+						Data = ComputeFlattened(Data, Primvar.Indices);
+					},
+					[&Primvar](pxr::VtArray<pxr::GfVec3f>& Data)
+					{
+						Data = ComputeFlattened(Data, Primvar.Indices);
+					},
+					[&Primvar](pxr::VtArray<pxr::GfVec4f>& Data)
+					{
+						Data = ComputeFlattened(Data, Primvar.Indices);
+					}
+				},
+				Primvar.Data
+			);
 
 			Primvar.Indices = {};
 		}
@@ -2516,14 +2544,35 @@ namespace UE::UsdGeomMeshConversion::Private
 			// Register the primvar if it doesn't exist
 			if (!OutMeshDescription.VertexInstanceAttributes().HasAttribute(*PrimvarName))
 			{
-				Visit(TOverloaded{
-					[&PrimvarName, &OutMeshDescription](const pxr::VtArray<bool>& ){OutMeshDescription.VertexInstanceAttributes().RegisterAttribute<bool>(*PrimvarName);},
-					[&PrimvarName, &OutMeshDescription](const pxr::VtArray<int32>& ){OutMeshDescription.VertexInstanceAttributes().RegisterAttribute<int32>(*PrimvarName);},
-					[&PrimvarName, &OutMeshDescription](const pxr::VtArray<float>& ){OutMeshDescription.VertexInstanceAttributes().RegisterAttribute<float>(*PrimvarName);},
-					[&PrimvarName, &OutMeshDescription](const pxr::VtArray<pxr::GfVec2f>& ){OutMeshDescription.VertexInstanceAttributes().RegisterAttribute<FVector2f>(*PrimvarName);},
-					[&PrimvarName, &OutMeshDescription](const pxr::VtArray<pxr::GfVec3f>& ){OutMeshDescription.VertexInstanceAttributes().RegisterAttribute<FVector3f>(*PrimvarName);},
-					[&PrimvarName, &OutMeshDescription](const pxr::VtArray<pxr::GfVec4f>& ){OutMeshDescription.VertexInstanceAttributes().RegisterAttribute<FVector4f>(*PrimvarName);}
-					}, CustomPrimvar.Data);
+				Visit(
+					TOverloaded{
+						[&PrimvarName, &OutMeshDescription](const pxr::VtArray<bool>&)
+						{
+							OutMeshDescription.VertexInstanceAttributes().RegisterAttribute<bool>(*PrimvarName);
+						},
+						[&PrimvarName, &OutMeshDescription](const pxr::VtArray<int32>&)
+						{
+							OutMeshDescription.VertexInstanceAttributes().RegisterAttribute<int32>(*PrimvarName);
+						},
+						[&PrimvarName, &OutMeshDescription](const pxr::VtArray<float>&)
+						{
+							OutMeshDescription.VertexInstanceAttributes().RegisterAttribute<float>(*PrimvarName);
+						},
+						[&PrimvarName, &OutMeshDescription](const pxr::VtArray<pxr::GfVec2f>&)
+						{
+							OutMeshDescription.VertexInstanceAttributes().RegisterAttribute<FVector2f>(*PrimvarName);
+						},
+						[&PrimvarName, &OutMeshDescription](const pxr::VtArray<pxr::GfVec3f>&)
+						{
+							OutMeshDescription.VertexInstanceAttributes().RegisterAttribute<FVector3f>(*PrimvarName);
+						},
+						[&PrimvarName, &OutMeshDescription](const pxr::VtArray<pxr::GfVec4f>&)
+						{
+							OutMeshDescription.VertexInstanceAttributes().RegisterAttribute<FVector4f>(*PrimvarName);
+						}
+					},
+					CustomPrimvar.Data
+				);
 			}
 		}
 	}
@@ -2792,7 +2841,7 @@ namespace UE::UsdGeomMeshConversion::Private
 																					.GetAttributesRef<FVector3f>(
 																						MeshAttribute::VertexInstance::Velocity
 																					);
-						
+
 			using FPrimvarVertexInstanceType = TVariant<
 				TVertexInstanceAttributesRef<bool>,
 				TVertexInstanceAttributesRef<int32>,
@@ -2805,25 +2854,76 @@ namespace UE::UsdGeomMeshConversion::Private
 
 			for (const TPair<FString, FUsdMeshData::FPrimvar>& Pair : InMeshData.Primvars)
 			{
-				Visit(TOverloaded{
-				[&MeshDescriptionPrimvar, &Pair, &OutMeshDescription](const pxr::VtArray<bool>&)
-				{MeshDescriptionPrimvar.Emplace(Pair.Key,FPrimvarVertexInstanceType{ TInPlaceType<TVertexInstanceAttributesRef<bool>>{}, OutMeshDescription.VertexInstanceAttributes().GetAttributesRef<bool>(*(Pair.Key))});},
+				Visit(
+					TOverloaded{
+						[&MeshDescriptionPrimvar, &Pair, &OutMeshDescription](const pxr::VtArray<bool>&)
+						{
+							MeshDescriptionPrimvar.Emplace(
+								Pair.Key,
+								FPrimvarVertexInstanceType{
+									TInPlaceType<TVertexInstanceAttributesRef<bool>>{},
+									OutMeshDescription.VertexInstanceAttributes().GetAttributesRef<bool>(*(Pair.Key))
+								}
+							);
+						},
 
-				[&MeshDescriptionPrimvar, &Pair, &OutMeshDescription](const pxr::VtArray<int32>&)
-				{MeshDescriptionPrimvar.Emplace(Pair.Key,FPrimvarVertexInstanceType{ TInPlaceType<TVertexInstanceAttributesRef<int32>>{}, OutMeshDescription.VertexInstanceAttributes().GetAttributesRef<int32>(*(Pair.Key))});},
+						[&MeshDescriptionPrimvar, &Pair, &OutMeshDescription](const pxr::VtArray<int32>&)
+						{
+							MeshDescriptionPrimvar.Emplace(
+								Pair.Key,
+								FPrimvarVertexInstanceType{
+									TInPlaceType<TVertexInstanceAttributesRef<int32>>{},
+									OutMeshDescription.VertexInstanceAttributes().GetAttributesRef<int32>(*(Pair.Key))
+								}
+							);
+						},
 
-				[&MeshDescriptionPrimvar, &Pair, &OutMeshDescription](const pxr::VtArray<float>&)
-				{MeshDescriptionPrimvar.Emplace(Pair.Key,FPrimvarVertexInstanceType{ TInPlaceType<TVertexInstanceAttributesRef<float>>{}, OutMeshDescription.VertexInstanceAttributes().GetAttributesRef<float>(*(Pair.Key))});},
+						[&MeshDescriptionPrimvar, &Pair, &OutMeshDescription](const pxr::VtArray<float>&)
+						{
+							MeshDescriptionPrimvar.Emplace(
+								Pair.Key,
+								FPrimvarVertexInstanceType{
+									TInPlaceType<TVertexInstanceAttributesRef<float>>{},
+									OutMeshDescription.VertexInstanceAttributes().GetAttributesRef<float>(*(Pair.Key))
+								}
+							);
+						},
 
-				[&MeshDescriptionPrimvar, &Pair, &OutMeshDescription](const pxr::VtArray<pxr::GfVec2f>&)
-				{MeshDescriptionPrimvar.Emplace(Pair.Key,FPrimvarVertexInstanceType{ TInPlaceType<TVertexInstanceAttributesRef<FVector2f>>{}, OutMeshDescription.VertexInstanceAttributes().GetAttributesRef<FVector2f>(*(Pair.Key))});},
+						[&MeshDescriptionPrimvar, &Pair, &OutMeshDescription](const pxr::VtArray<pxr::GfVec2f>&)
+						{
+							MeshDescriptionPrimvar.Emplace(
+								Pair.Key,
+								FPrimvarVertexInstanceType{
+									TInPlaceType<TVertexInstanceAttributesRef<FVector2f>>{},
+									OutMeshDescription.VertexInstanceAttributes().GetAttributesRef<FVector2f>(*(Pair.Key))
+								}
+							);
+						},
 
-				[&MeshDescriptionPrimvar, &Pair, &OutMeshDescription](const pxr::VtArray<pxr::GfVec3f>&)
-				{MeshDescriptionPrimvar.Emplace(Pair.Key,FPrimvarVertexInstanceType{ TInPlaceType<TVertexInstanceAttributesRef<FVector3f>>{}, OutMeshDescription.VertexInstanceAttributes().GetAttributesRef<FVector3f>(*(Pair.Key))});},
+						[&MeshDescriptionPrimvar, &Pair, &OutMeshDescription](const pxr::VtArray<pxr::GfVec3f>&)
+						{
+							MeshDescriptionPrimvar.Emplace(
+								Pair.Key,
+								FPrimvarVertexInstanceType{
+									TInPlaceType<TVertexInstanceAttributesRef<FVector3f>>{},
+									OutMeshDescription.VertexInstanceAttributes().GetAttributesRef<FVector3f>(*(Pair.Key))
+								}
+							);
+						},
 
-				[&MeshDescriptionPrimvar, &Pair, &OutMeshDescription](const pxr::VtArray<pxr::GfVec4f>&)
-				{MeshDescriptionPrimvar.Emplace(Pair.Key,FPrimvarVertexInstanceType{ TInPlaceType<TVertexInstanceAttributesRef<FVector4f>>{}, OutMeshDescription.VertexInstanceAttributes().GetAttributesRef<FVector4f>(*(Pair.Key))});}
-				}, Pair.Value.Data);
+						[&MeshDescriptionPrimvar, &Pair, &OutMeshDescription](const pxr::VtArray<pxr::GfVec4f>&)
+						{
+							MeshDescriptionPrimvar.Emplace(
+								Pair.Key,
+								FPrimvarVertexInstanceType{
+									TInPlaceType<TVertexInstanceAttributesRef<FVector4f>>{},
+									OutMeshDescription.VertexInstanceAttributes().GetAttributesRef<FVector4f>(*(Pair.Key))
+								}
+							);
+						}
+					},
+					Pair.Value.Data
+				);
 			}
 
 			OutMeshDescription.ReserveNewVertexInstances(InMeshData.FaceVertexCounts.size() * 3);
@@ -2997,65 +3097,79 @@ namespace UE::UsdGeomMeshConversion::Private
 						}
 					}
 
-					//Primvars
+					// Primvars
 					for (const TPair<FString, FUsdMeshData::FPrimvar>& Pair : InMeshData.Primvars)
 					{
 						const FString& PrimvarName = Pair.Key;
 						const FUsdMeshData::FPrimvar& CustomPrimvar = Pair.Value;
-						const size_t ValueIndex = GetPrimValueIndex(CustomPrimvar.Interpolation,
-																	ControlPointIndex,
-																	CurrentVertexInstanceIndex,
-																	PolygonIndex);
-
-						// We visit InMeshData.Primvar.Data for each type and make sure to fill the correct TVertexInstanceAttributesRef with the right type of the primvar
-						Visit([ValueIndex, AddedVertexInstanceId, &PrimvarName, &MeshDescriptionPrimvar, &InMeshData](auto && Value)
-							  {
-								  if (!Value.empty())
-								  {
-									  if (Value.size() > ValueIndex)
-									  {
-										  // underlying type of the pxr::VtArray that the primvar holds
-										  using FPrimvarType = std::decay_t<decltype(Value)>::ElementType;
-
-										  FPrimvarVertexInstanceType& PrimvarVertexInstance = MeshDescriptionPrimvar.FindChecked(PrimvarName);
-
-										  if constexpr (std::is_same_v<FPrimvarType, pxr::GfVec2f>)
-										  {
-											  TVertexInstanceAttributesRef<FVector2f>& VertexInstance = PrimvarVertexInstance.Get<TVertexInstanceAttributesRef<FVector2f>>();
-											  VertexInstance[AddedVertexInstanceId] = FVector2f{ Value[ValueIndex][0], Value[ValueIndex][1] };
-										  }
-										  else if constexpr (std::is_same_v<FPrimvarType, pxr::GfVec3f>)
-										  {
-											  TVertexInstanceAttributesRef<FVector3f>& VertexInstance = PrimvarVertexInstance.Get<TVertexInstanceAttributesRef<FVector3f>>();
-											  VertexInstance[AddedVertexInstanceId] = FVector3f{ Value[ValueIndex][0], Value[ValueIndex][1], Value[ValueIndex][2] };
-										  }
-										  else if constexpr (std::is_same_v<FPrimvarType, pxr::GfVec4f>)
-										  {
-											  TVertexInstanceAttributesRef<FVector4f>& VertexInstance = PrimvarVertexInstance.Get<TVertexInstanceAttributesRef<FVector4f>>();
-											  VertexInstance[AddedVertexInstanceId] = FVector4f{ Value[ValueIndex][0], Value[ValueIndex][1], Value[ValueIndex][2], Value[ValueIndex][3] };
-
-										  }
-										  else if constexpr (std::is_arithmetic_v<FPrimvarType>)
-										  {
-											  TVertexInstanceAttributesRef<FPrimvarType>& VertexInstance = PrimvarVertexInstance.Get<TVertexInstanceAttributesRef<FPrimvarType>>();
-											  VertexInstance[AddedVertexInstanceId] = Value[ValueIndex];
-										  }
-									  }
-									  else
-									  {
-										  UE_LOG(
-											  LogUsd,
-											  Warning,
-											  TEXT("Trying to read Primvar '%s' at index %u from prim '%s' but the prim only has %u values! Using "
-											  "zeros instead."),
-											  *PrimvarName,
-											  ValueIndex,
-											  *InMeshData.SourcePrimPath,
-											  InMeshData.DisplayColors.size()
-										  );
-									  }
-								  }
-							  }, CustomPrimvar.Data);
+						const size_t
+							ValueIndex = GetPrimValueIndex(CustomPrimvar.Interpolation, ControlPointIndex, CurrentVertexInstanceIndex, PolygonIndex);
+
+						// We visit InMeshData.Primvar.Data for each type and make sure to fill the correct TVertexInstanceAttributesRef with the
+						// right type of the primvar
+						Visit(
+							[ValueIndex, AddedVertexInstanceId, &PrimvarName, &MeshDescriptionPrimvar, &InMeshData](auto&& Value)
+							{
+								if (!Value.empty())
+								{
+									if (Value.size() > ValueIndex)
+									{
+										// underlying type of the pxr::VtArray that the primvar holds
+										using FPrimvarType = std::decay_t<decltype(Value)>::ElementType;
+
+										FPrimvarVertexInstanceType& PrimvarVertexInstance = MeshDescriptionPrimvar.FindChecked(PrimvarName);
+
+										if constexpr (std::is_same_v<FPrimvarType, pxr::GfVec2f>)
+										{
+											TVertexInstanceAttributesRef<FVector2f>&
+												VertexInstance = PrimvarVertexInstance.Get<TVertexInstanceAttributesRef<FVector2f>>();
+											VertexInstance[AddedVertexInstanceId] = FVector2f{Value[ValueIndex][0], Value[ValueIndex][1]};
+										}
+										else if constexpr (std::is_same_v<FPrimvarType, pxr::GfVec3f>)
+										{
+											TVertexInstanceAttributesRef<FVector3f>&
+												VertexInstance = PrimvarVertexInstance.Get<TVertexInstanceAttributesRef<FVector3f>>();
+											VertexInstance[AddedVertexInstanceId] = FVector3f{
+												Value[ValueIndex][0],
+												Value[ValueIndex][1],
+												Value[ValueIndex][2]
+											};
+										}
+										else if constexpr (std::is_same_v<FPrimvarType, pxr::GfVec4f>)
+										{
+											TVertexInstanceAttributesRef<FVector4f>&
+												VertexInstance = PrimvarVertexInstance.Get<TVertexInstanceAttributesRef<FVector4f>>();
+											VertexInstance[AddedVertexInstanceId] = FVector4f{
+												Value[ValueIndex][0],
+												Value[ValueIndex][1],
+												Value[ValueIndex][2],
+												Value[ValueIndex][3]
+											};
+										}
+										else if constexpr (std::is_arithmetic_v<FPrimvarType>)
+										{
+											TVertexInstanceAttributesRef<FPrimvarType>&
+												VertexInstance = PrimvarVertexInstance.Get<TVertexInstanceAttributesRef<FPrimvarType>>();
+											VertexInstance[AddedVertexInstanceId] = Value[ValueIndex];
+										}
+									}
+									else
+									{
+										UE_LOG(
+											LogUsd,
+											Warning,
+											TEXT("Trying to read Primvar '%s' at index %u from prim '%s' but the prim only has %u values! Using "
+												 "zeros instead."),
+											*PrimvarName,
+											ValueIndex,
+											*InMeshData.SourcePrimPath,
+											InMeshData.DisplayColors.size()
+										);
+									}
+								}
+							},
+							CustomPrimvar.Data
+						);
 					}
 				}
 
@@ -5896,6 +6010,21 @@ bool UsdUtils::IsGeomMeshALOD(const pxr::UsdPrim& UsdMeshPrim)
 	return UE::UsdGeomMeshConversion::Private::DoesPrimContainMeshLODsInternal(UsdMeshPrim.GetParent());
 }
 
+bool UsdUtils::IsCollisionEnabledForPrim(const pxr::UsdPrim& UsdPrim)
+{
+	if (pxr::UsdPhysicsCollisionAPI CollisionAPI{UsdPrim})
+	{
+		bool bIsCollisionEnabled = false;
+		pxr::UsdAttribute CollisionAttr = CollisionAPI.GetCollisionEnabledAttr();
+		if (CollisionAttr && CollisionAttr.Get(&bIsCollisionEnabled) && bIsCollisionEnabled)
+		{
+			return true;
+		}
+	}
+
+	return false;
+}
+
 bool UsdUtils::IsCollisionMesh(const pxr::UsdPrim& UsdPrim)
 {
 	// https://openusd.org/release/api/usd_physics_page_front.html#usdPhysics_collision_shapes
@@ -5910,20 +6039,7 @@ bool UsdUtils::IsCollisionMesh(const pxr::UsdPrim& UsdPrim)
 		return false;
 	}
 
-	bool bIsCollisionEnabled = false;
-	if (pxr::UsdPhysicsCollisionAPI CollisionAPI{UsdPrim})
-	{
-		if (pxr::UsdAttribute CollisionAttr = CollisionAPI.GetCollisionEnabledAttr())
-		{
-			CollisionAttr.Get(&bIsCollisionEnabled);
-		}
-
-		if (!bIsCollisionEnabled)
-		{
-			return false;
-		}
-	}
-	else
+	if (!IsCollisionEnabledForPrim(UsdPrim))
 	{
 		return false;
 	}
@@ -5959,6 +6075,77 @@ bool UsdUtils::IsCollisionMesh(const pxr::UsdPrim& UsdPrim)
 	return true;
 }
 
+UsdUtils::EUsdCollisionType UsdUtils::GetCollisionApproximationType(const pxr::UsdPrim& UsdPrim)
+{
+	EUsdCollisionType Approximation(EUsdCollisionType::None);
+	if (pxr::UsdGeomMesh UsdMesh{UsdPrim})
+	{
+		// Get the collision approximation type (only meshes should have UsdPhysicsMeshCollisionAPI)
+		if (pxr::UsdPhysicsMeshCollisionAPI MeshCollisionAPI{UsdMesh})
+		{
+			pxr::TfToken ApproximationValue{pxr::UsdPhysicsTokens->none};
+			pxr::UsdAttribute Attr = MeshCollisionAPI.GetApproximationAttr();
+			if (Attr && Attr.Get(&ApproximationValue))
+			{
+				if (ApproximationValue == pxr::UsdPhysicsTokens->convexDecomposition)
+				{
+#if WITH_EDITOR
+					return EUsdCollisionType::ConvexDecomposition;
+#else
+					return EUsdCollisionType::ConvexHull;
+#endif
+				}
+				else if (ApproximationValue == pxr::UsdPhysicsTokens->convexHull)
+				{
+					return EUsdCollisionType::ConvexHull;
+				}
+				else if (ApproximationValue == pxr::UsdPhysicsTokens->boundingSphere)
+				{
+					return EUsdCollisionType::Sphere;
+				}
+				else if (ApproximationValue == pxr::UsdPhysicsTokens->boundingCube)
+				{
+					return EUsdCollisionType::Cube;
+				}
+				else if (ApproximationValue == pxr::UsdPhysicsTokens->meshSimplification)
+				{
+					return EUsdCollisionType::MeshSimplification;
+				}
+			}
+		}
+	}
+	else
+	{
+		// Collision for primitives are converted to their closest approximation
+		if (UsdPrim.IsA(pxr::UsdGeomTokens->Capsule))
+		{
+			return EUsdCollisionType::Capsule;
+		}
+		else if (UsdPrim.IsA(pxr::UsdGeomTokens->Cone))
+		{
+			return EUsdCollisionType::CustomMesh;
+		}
+		else if (UsdPrim.IsA(pxr::UsdGeomTokens->Cube))
+		{
+			return EUsdCollisionType::Cube;
+		}
+		else if (UsdPrim.IsA(pxr::UsdGeomTokens->Cylinder))
+		{
+			return EUsdCollisionType::CustomMesh;
+		}
+		else if (UsdPrim.IsA(pxr::UsdGeomTokens->Sphere))
+		{
+			return EUsdCollisionType::Sphere;
+		}
+		else if (UsdPrim.IsA(pxr::UsdGeomTokens->Plane))
+		{
+			return EUsdCollisionType::CustomMesh;
+		}
+	}
+
+	return EUsdCollisionType::None;
+}
+
 int32 UsdUtils::GetNumberOfLODVariants(const pxr::UsdPrim& Prim)
 {
 	FScopedUsdAllocs Allocs;
diff --git a/Engine/Plugins/Runtime/USDCore/Source/USDUtilities/Public/USDGeomMeshConversion.h b/Engine/Plugins/Runtime/USDCore/Source/USDUtilities/Public/USDGeomMeshConversion.h
--- a/Engine/Plugins/Runtime/USDCore/Source/USDUtilities/Public/USDGeomMeshConversion.h
+++ b/Engine/Plugins/Runtime/USDCore/Source/USDUtilities/Public/USDGeomMeshConversion.h
@@ -108,7 +108,8 @@ namespace UsdToUnreal
 	 * @param InOutMaterialAssignments - Input/Output parameter that will be filled with the material data extracted from Mesh
 	 * @param CommonOptions - Shared options used for the conversion
 	 * @param SkinningQuery - For computing Skinned Points, if the target and source data set is for a SkeletalMesh.
-	 * @param SkeletonQuery - For computing Skinned Points, if the target and source data set is for a SkeletalMesh. (primarily for MeshToSkeletonRestPose calculations)
+	 * @param SkeletonQuery - For computing Skinned Points, if the target and source data set is for a SkeletalMesh. (primarily for
+	 * MeshToSkeletonRestPose calculations)
 	 * @return Whether the conversion was successful or not.
 	 */
 	USDUTILITIES_API bool ConvertGeomMesh(
@@ -431,9 +432,27 @@ namespace UsdUtils
 	/** Returns whether this UsdMesh can be interpreted as a LOD of a mesh with multiple LODs (i.e. if the Mesh prim is inside a LOD variant) */
 	USDUTILITIES_API bool IsGeomMeshALOD(const pxr::UsdPrim& UsdMeshPrim);
 
+	// Approximation types from PhysicsMeshCollisionAPI and UE-specific approximations
+	enum class EUsdCollisionType : uint8
+	{
+		None,	 // no approximation so equivalent to CTF_UseComplexAsSimple
+		ConvexDecomposition,
+		ConvexHull,
+		Sphere,
+		Cube,
+		MeshSimplification,
+		Capsule,
+		CustomMesh
+	};
+
+	/** Returns whether this UsdPrim has the UsdPhysicsCollisionAPI schema and has enabled collisions */
+	USDUTILITIES_API bool IsCollisionEnabledForPrim(const pxr::UsdPrim& UsdPrim);
+
 	/** Returns whether this UsdPrim is a custom collision mesh */
 	USDUTILITIES_API bool IsCollisionMesh(const pxr::UsdPrim& UsdPrim);
 
+	USDUTILITIES_API EUsdCollisionType GetCollisionApproximationType(const pxr::UsdPrim& UsdPrim);
+
 	/** Returns how many LOD variants the Prim has. Note that this will return 0 if called on one of the LOD meshes themselves, it's meant to be
 	 * called on its parent */
 	USDUTILITIES_API int32 GetNumberOfLODVariants(const pxr::UsdPrim& Prim);
