USD Core:
- Added Geomsubset wrapper.
- Added attribute ClearConnections wrapper function.

#rnx

diff --git a/Engine/Plugins/Runtime/USDCore/Source/UnrealUSDWrapper/Private/UsdWrappers/UsdAttribute.cpp b/Engine/Plugins/Runtime/USDCore/Source/UnrealUSDWrapper/Private/UsdWrappers/UsdAttribute.cpp
--- a/Engine/Plugins/Runtime/USDCore/Source/UnrealUSDWrapper/Private/UsdWrappers/UsdAttribute.cpp
+++ b/Engine/Plugins/Runtime/USDCore/Source/UnrealUSDWrapper/Private/UsdWrappers/UsdAttribute.cpp
@@ -332,6 +332,15 @@ namespace UE
 #endif	  // #if USE_USD_SDK
 	}
 
+	bool FUsdAttribute::ClearConnections() const
+	{
+#if USE_USD_SDK
+		return Impl->PxrUsdAttribute.Get().ClearConnections();
+#else
+		return false;
+#endif	  // #if USE_USD_SDK
+	}
+
 	bool FUsdAttribute::GetUnionedTimeSamples(const TArray<UE::FUsdAttribute>& Attrs, TArray<double>& OutTimes)
 	{
 		bool bResult = false;
diff --git a/Engine/Plugins/Runtime/USDCore/Source/UnrealUSDWrapper/Private/UsdWrappers/UsdGeomSubset.cpp b/Engine/Plugins/Runtime/USDCore/Source/UnrealUSDWrapper/Private/UsdWrappers/UsdGeomSubset.cpp
new file mode 100644
--- /dev/null
+++ b/Engine/Plugins/Runtime/USDCore/Source/UnrealUSDWrapper/Private/UsdWrappers/UsdGeomSubset.cpp
@@ -0,0 +1,167 @@
+// Copyright Epic Games, Inc. All Rights Reserved.
+
+#include "UsdWrappers/UsdGeomSubset.h"
+
+#include "USDMemory.h"
+
+#include "UsdWrappers/UsdPrim.h"
+
+#if USE_USD_SDK
+#include "USDIncludesStart.h"
+#include "pxr/usd/usd/attribute.h"
+#include "pxr/usd/usdGeom/subset.h"
+#include "USDIncludesEnd.h"
+#endif	  // #if USE_USD_SDK
+
+namespace UE
+{
+	namespace Internal
+	{
+		class FUsdGeomSubsetImpl
+		{
+		public:
+			FUsdGeomSubsetImpl() = default;
+
+#if USE_USD_SDK
+			explicit FUsdGeomSubsetImpl(const pxr::UsdPrim& InUsdPrim)
+				: PxrUsdGeomSubset(MakeUsdStore<pxr::UsdGeomSubset>(InUsdPrim))
+			{
+			}
+
+			explicit FUsdGeomSubsetImpl(const pxr::UsdGeomSubset& InUsdGeomSubset)
+				: PxrUsdGeomSubset(InUsdGeomSubset)
+			{
+			}
+
+			explicit FUsdGeomSubsetImpl(pxr::UsdGeomSubset&& InUsdGeomSubset)
+				: PxrUsdGeomSubset(MoveTemp(InUsdGeomSubset))
+			{
+			}
+
+			TUsdStore<pxr::UsdGeomSubset> PxrUsdGeomSubset;
+#endif	  // #if USE_USD_SDK
+		};
+	}	 // namespace Internal
+
+	FUsdGeomSubset::FUsdGeomSubset()
+	{
+		FScopedUnrealAllocs UnrealAllocs;
+		Impl = MakeUnique<Internal::FUsdGeomSubsetImpl>();
+	}
+
+	FUsdGeomSubset::FUsdGeomSubset(const FUsdPrim& UsdPrim)
+	{
+#if USE_USD_SDK
+		FScopedUnrealAllocs UnrealAllocs;
+		Impl = MakeUnique<Internal::FUsdGeomSubsetImpl>(UsdPrim);
+#endif	  // #if USE_USD_SDK
+	}
+
+	FUsdGeomSubset::FUsdGeomSubset(const FUsdGeomSubset& Other)
+	{
+#if USE_USD_SDK
+		FScopedUnrealAllocs UnrealAllocs;
+		Impl = MakeUnique<Internal::FUsdGeomSubsetImpl>(Other.Impl->PxrUsdGeomSubset.Get());
+#endif	  // #if USE_USD_SDK
+	}
+
+	FUsdGeomSubset::FUsdGeomSubset(FUsdGeomSubset&& Other) = default;
+
+	FUsdGeomSubset::~FUsdGeomSubset()
+	{
+		FScopedUnrealAllocs UnrealAllocs;
+		Impl.Reset();
+	}
+
+	FUsdGeomSubset& FUsdGeomSubset::operator=(const FUsdGeomSubset& Other)
+	{
+#if USE_USD_SDK
+		FScopedUnrealAllocs UnrealAllocs;
+		Impl = MakeUnique<Internal::FUsdGeomSubsetImpl>(Other.Impl->PxrUsdGeomSubset.Get());
+#endif	  // #if USE_USD_SDK
+		return *this;
+	}
+
+	FUsdGeomSubset& FUsdGeomSubset::operator=(FUsdGeomSubset&& Other)
+	{
+		FScopedUnrealAllocs UnrealAllocs;
+		Impl = MoveTemp(Other.Impl);
+
+		return *this;
+	}
+
+	FUsdGeomSubset::operator bool() const
+	{
+#if USE_USD_SDK
+		return (bool)Impl->PxrUsdGeomSubset.Get();
+#else
+		return false;
+#endif	  // #if USE_USD_SDK
+	}
+
+#if USE_USD_SDK
+	FUsdGeomSubset::FUsdGeomSubset(const pxr::UsdGeomSubset& InUsdGeomSubset)
+	{
+		FScopedUnrealAllocs UnrealAllocs;
+		Impl = MakeUnique<Internal::FUsdGeomSubsetImpl>(InUsdGeomSubset);
+	}
+
+	FUsdGeomSubset::FUsdGeomSubset(pxr::UsdGeomSubset&& InUsdGeomSubset)
+	{
+		FScopedUnrealAllocs UnrealAllocs;
+		Impl = MakeUnique<Internal::FUsdGeomSubsetImpl>(MoveTemp(InUsdGeomSubset));
+	}
+
+	FUsdGeomSubset& FUsdGeomSubset::operator=(const pxr::UsdGeomSubset& InUsdGeomSubset)
+	{
+		FScopedUnrealAllocs UnrealAllocs;
+		Impl = MakeUnique<Internal::FUsdGeomSubsetImpl>(InUsdGeomSubset);
+		return *this;
+	}
+
+	FUsdGeomSubset& FUsdGeomSubset::operator=(pxr::UsdGeomSubset&& InUsdGeomSubset)
+	{
+		FScopedUnrealAllocs UnrealAllocs;
+		Impl = MakeUnique<Internal::FUsdGeomSubsetImpl>(MoveTemp(InUsdGeomSubset));
+		return *this;
+	}
+
+	FUsdGeomSubset::operator pxr::UsdGeomSubset&()
+	{
+		return Impl->PxrUsdGeomSubset.Get();
+	}
+
+	FUsdGeomSubset::operator const pxr::UsdGeomSubset&() const
+	{
+		return Impl->PxrUsdGeomSubset.Get();
+	}
+#endif	  // #if USE_USD_SDK
+
+	FUsdAttribute FUsdGeomSubset::GetElementTypeAttr() const
+	{
+#if USE_USD_SDK
+		return FUsdAttribute(Impl->PxrUsdGeomSubset.Get().GetElementTypeAttr());
+#else
+		return FUsdAttribute();
+#endif	  // #if USE_USD_SDK
+	}
+
+	FUsdAttribute FUsdGeomSubset::GetFamilyNameAttr() const
+	{
+#if USE_USD_SDK
+		return FUsdAttribute(Impl->PxrUsdGeomSubset.Get().GetFamilyNameAttr());
+#else
+		return FUsdAttribute();
+#endif	  // #if USE_USD_SDK
+	}
+
+	FUsdAttribute FUsdGeomSubset::GetIndicesAttr() const
+	{
+#if USE_USD_SDK
+		return FUsdAttribute(Impl->PxrUsdGeomSubset.Get().GetIndicesAttr());
+#else
+		return FUsdAttribute();
+#endif	  // #if USE_USD_SDK
+	}
+
+}	 // namespace UE
diff --git a/Engine/Plugins/Runtime/USDCore/Source/UnrealUSDWrapper/Public/UsdWrappers/UsdAttribute.h b/Engine/Plugins/Runtime/USDCore/Source/UnrealUSDWrapper/Public/UsdWrappers/UsdAttribute.h
--- a/Engine/Plugins/Runtime/USDCore/Source/UnrealUSDWrapper/Public/UsdWrappers/UsdAttribute.h
+++ b/Engine/Plugins/Runtime/USDCore/Source/UnrealUSDWrapper/Public/UsdWrappers/UsdAttribute.h
@@ -98,6 +98,8 @@ namespace UE
 		bool Clear() const;
 		bool ClearAtTime(double Time) const;
 
+		bool ClearConnections() const;
+
 		static bool GetUnionedTimeSamples(const TArray<UE::FUsdAttribute>& Attrs, TArray<double>& OutTimes);
 
 		FSdfPath GetPath() const;
diff --git a/Engine/Plugins/Runtime/USDCore/Source/UnrealUSDWrapper/Public/UsdWrappers/UsdGeomSubset.h b/Engine/Plugins/Runtime/USDCore/Source/UnrealUSDWrapper/Public/UsdWrappers/UsdGeomSubset.h
new file mode 100644
--- /dev/null
+++ b/Engine/Plugins/Runtime/USDCore/Source/UnrealUSDWrapper/Public/UsdWrappers/UsdGeomSubset.h
@@ -0,0 +1,64 @@
+// Copyright Epic Games, Inc. All Rights Reserved.
+
+#pragma once
+
+#include "Templates/UniquePtr.h"
+#include "UsdWrappers/UsdAttribute.h"
+
+#if USE_USD_SDK
+PXR_NAMESPACE_OPEN_SCOPE
+	class UsdGeomSubset;
+PXR_NAMESPACE_CLOSE_SCOPE
+#endif	  // #if USE_USD_SDK
+
+namespace UE
+{
+	class FUsdPrim;
+
+	namespace Internal
+	{
+		class FUsdGeomSubsetImpl;
+	}
+
+	/**
+	 * Minimal pxr::UsdGeomSubset wrapper for Unreal that can be used from no-rtti modules.
+	 */
+	class UNREALUSDWRAPPER_API FUsdGeomSubset
+	{
+	public:
+		FUsdGeomSubset();
+		FUsdGeomSubset(const FUsdPrim& UsdPrim);
+
+		FUsdGeomSubset(const FUsdGeomSubset& Other);
+		FUsdGeomSubset(FUsdGeomSubset&& Other);
+		~FUsdGeomSubset();
+
+		FUsdGeomSubset& operator=(const FUsdGeomSubset& Other);
+		FUsdGeomSubset& operator=(FUsdGeomSubset&& Other);
+
+		explicit operator bool() const;
+
+		// Auto conversion from/to pxr::FUsdGeomSubset
+	public:
+#if USE_USD_SDK
+		explicit FUsdGeomSubset(const pxr::UsdGeomSubset& InUsdGeomSubset);
+		explicit FUsdGeomSubset(pxr::UsdGeomSubset&& InUsdGeomSubset);
+		FUsdGeomSubset& operator=(const pxr::UsdGeomSubset& InUsdGeomSubset);
+		FUsdGeomSubset& operator=(pxr::UsdGeomSubset&& InUsdGeomSubset);
+
+		operator pxr::UsdGeomSubset&();
+		operator const pxr::UsdGeomSubset&() const;
+#endif	  // #if USE_USD_SDK
+
+		// Wrapped pxr::UsdGeomSubset functions, refer to the USD SDK documentation
+	public:
+		FUsdAttribute GetElementTypeAttr() const;
+
+		FUsdAttribute GetFamilyNameAttr() const;
+
+		FUsdAttribute GetIndicesAttr() const;
+
+	private:
+		TUniquePtr<Internal::FUsdGeomSubsetImpl> Impl;
+	};
+}	 // namespace UE
