USD: Sparse type name editing for add reference, add payload and export workflows.

diff --git a/Engine/Plugins/Importers/USDImporter/Source/USDExporter/Private/AnimSequenceExporterUSD.cpp b/Engine/Plugins/Importers/USDImporter/Source/USDExporter/Private/AnimSequenceExporterUSD.cpp
--- a/Engine/Plugins/Importers/USDImporter/Source/USDExporter/Private/AnimSequenceExporterUSD.cpp
+++ b/Engine/Plugins/Importers/USDImporter/Source/USDExporter/Private/AnimSequenceExporterUSD.cpp
@@ -335,7 +335,7 @@ bool UAnimSequenceExporterUSD::ExportBinary(
 		}
 		else
 		{
-			SkelAnimPrim = AnimationStage.DefinePrim(SkelAnimPath, TEXT("Mesh"));
+			SkelAnimPrim = AnimationStage.DefinePrim(SkelAnimPath);
 			if (!SkelAnimPrim)
 			{
 				return false;
diff --git a/Engine/Plugins/Importers/USDImporter/Source/USDExporter/Private/LevelSequenceExporterUSD.cpp b/Engine/Plugins/Importers/USDImporter/Source/USDExporter/Private/LevelSequenceExporterUSD.cpp
--- a/Engine/Plugins/Importers/USDImporter/Source/USDExporter/Private/LevelSequenceExporterUSD.cpp
+++ b/Engine/Plugins/Importers/USDImporter/Source/USDExporter/Private/LevelSequenceExporterUSD.cpp
@@ -1491,7 +1491,7 @@ namespace UE::LevelSequenceExporterUSD::Private
 			}
 		}
 
-		// Convert all prims back to overs (so that this layer doesn't define anything on a stage that doesn't
+		// Convert all prims back to typeless overs (so that this layer doesn't define anything on a stage that doesn't
 		// previously have it - it's only supposed to carry animation data)
 		// We need to do this in this way because when going from 'def' to 'over' we need to do it from leaf towards the
 		// root, as USD doesn't like a parent 'over' with a child 'def'
@@ -1503,6 +1503,7 @@ namespace UE::LevelSequenceExporterUSD::Private
 				RecursiveSetOver(Child);
 			}
 
+			Prim.ClearTypeName();
 			Prim.SetSpecifier(ESdfSpecifier::Over);
 		};
 		FUsdPrim Root = UsdStage.GetPseudoRoot();
diff --git a/Engine/Plugins/Importers/USDImporter/Source/USDExporter/Private/USDConversionBlueprintLibrary.cpp b/Engine/Plugins/Importers/USDImporter/Source/USDExporter/Private/USDConversionBlueprintLibrary.cpp
--- a/Engine/Plugins/Importers/USDImporter/Source/USDExporter/Private/USDConversionBlueprintLibrary.cpp
+++ b/Engine/Plugins/Importers/USDImporter/Source/USDExporter/Private/USDConversionBlueprintLibrary.cpp
@@ -434,7 +434,12 @@ void UUsdConversionBlueprintLibrary::InsertSubLayer(const FString& ParentLayerPa
 void UUsdConversionBlueprintLibrary::AddReference(
 	const FString& ReferencingStagePath,
 	const FString& ReferencingPrimPath,
-	const FString& TargetStagePath
+	const FString& TargetStagePath,
+	const FString& TargetPrimPath,
+	double TimeCodeOffset,
+	double TimeCodeScale,
+	bool bUseProjectDefaultTypeHandling,
+	EReferencerTypeHandling ReferencerTypeHandling
 )
 {
 #if USE_USD_SDK
@@ -446,7 +451,14 @@ void UUsdConversionBlueprintLibrary::AddReference(
 	{
 		if (UE::FUsdPrim ReferencingPrim = ReferencingStage.GetPrimAtPath(UE::FSdfPath(*ReferencingPrimPath)))
 		{
-			UsdUtils::AddReference(ReferencingPrim, *TargetStagePath);
+			UsdUtils::AddReference(
+				ReferencingPrim,
+				*TargetStagePath,
+				bUseProjectDefaultTypeHandling ? TOptional<EReferencerTypeHandling>{} : ReferencerTypeHandling,
+				TargetPrimPath.IsEmpty() ? UE::FSdfPath{} : UE::FSdfPath{*TargetPrimPath},
+				TimeCodeOffset,
+				TimeCodeScale
+			);
 		}
 	}
 
@@ -461,7 +473,12 @@ void UUsdConversionBlueprintLibrary::AddReference(
 void UUsdConversionBlueprintLibrary::AddPayload(
 	const FString& ReferencingStagePath,
 	const FString& ReferencingPrimPath,
-	const FString& TargetStagePath
+	const FString& TargetStagePath,
+	const FString& TargetPrimPath,
+	double TimeCodeOffset,
+	double TimeCodeScale,
+	bool bUseProjectDefaultTypeHandling,
+	EReferencerTypeHandling ReferencerTypeHandling
 )
 {
 #if USE_USD_SDK
@@ -473,7 +490,14 @@ void UUsdConversionBlueprintLibrary::AddPayload(
 	{
 		if (UE::FUsdPrim ReferencingPrim = ReferencingStage.GetPrimAtPath(UE::FSdfPath(*ReferencingPrimPath)))
 		{
-			UsdUtils::AddPayload(ReferencingPrim, *TargetStagePath);
+			UsdUtils::AddPayload(
+				ReferencingPrim,
+				*TargetStagePath,
+				bUseProjectDefaultTypeHandling ? TOptional<EReferencerTypeHandling>{} : ReferencerTypeHandling,
+				TargetPrimPath.IsEmpty() ? UE::FSdfPath{} : UE::FSdfPath{*TargetPrimPath},
+				TimeCodeOffset,
+				TimeCodeScale
+			);
 		}
 	}
 
diff --git a/Engine/Plugins/Importers/USDImporter/Source/USDExporter/Public/USDConversionBlueprintLibrary.h b/Engine/Plugins/Importers/USDImporter/Source/USDExporter/Public/USDConversionBlueprintLibrary.h
--- a/Engine/Plugins/Importers/USDImporter/Source/USDExporter/Public/USDConversionBlueprintLibrary.h
+++ b/Engine/Plugins/Importers/USDImporter/Source/USDExporter/Public/USDConversionBlueprintLibrary.h
@@ -3,6 +3,7 @@
 #pragma once
 
 #include "USDMetadata.h"
+#include "USDProjectSettings.h"
 
 #include "AnalyticsBlueprintLibrary.h"
 #include "Engine/EngineTypes.h"
@@ -90,11 +91,59 @@ public:
 	UFUNCTION(BlueprintCallable, Category = "USD|Layer utils")
 	static void InsertSubLayer(const FString& ParentLayerPath, const FString& SubLayerPath, int32 Index = -1);
 
-	UFUNCTION(BlueprintCallable, Category = "USD|Layer utils")
-	static void AddReference(const FString& ReferencingStagePath, const FString& ReferencingPrimPath, const FString& TargetStagePath);
+	/**
+	 * Adds a reference to the layer at AbsoluteFilePath, optionally specifying a target prim inside that layer.
+	 * @param ReferencingStagePath - Stage identifier that contains the referencer prim
+	 * @param ReferencingPrimPath - Path to the referencer prim inside the referencer stage
+	 * @param TargetStagePath - Path to the layer to reference (can be empty for internal references). e.g. "C:\MyFolder\Cube.usda"
+	 * @param TargetPrimPath - Path to a specific prim to reference within the target layer. e.g. "/SomeOtherPrim"
+	 * @param TimeCodeOffset - Offset in timeCodes to author on the reference
+	 * @param TimeCodeScale - TimeCode scale to author on the reference
+	 * @param bUseProjectDefaultTypeHandling - If true it means we will use the ReferencerTypeHandling set on the USD project settings.
+	 *                                         If false it means we will use the provided ReferencerSchemaHandling
+	 * @param ReferencerSchemaHandling - How to behave when Prim has a different type than the referenced prim.
+	 *                                   Passing an unset TOptional will check the project settings for the value instead.
+	 *                                   The default value is MatchReferencedType for backwards compatibility.
+	 *                                   Note: If this is set to ShowDialog, the call may create an Editor transaction.
+	 */
+	 UFUNCTION(BlueprintCallable, Category = "USD|Layer utils")
+	 static void AddReference(
+		 const FString& ReferencingStagePath,
+		 const FString& ReferencingPrimPath,
+		 const FString& TargetStagePath,
+		 const FString& TargetPrimPath = TEXT(""),
+		 double TimeCodeOffset = 0.0,
+		 double TimeCodeScale = 1.0,
+		 bool bUseProjectDefaultTypeHandling = false,
+		 EReferencerTypeHandling ReferencerTypeHandling = EReferencerTypeHandling::MatchReferencedType
+		);
 
+	/**
+	 * Adds a payload to the layer at AbsoluteFilePath, optionally specifying a target prim inside that layer.
+	 * @param ReferencingStagePath - Stage identifier that contains the prim to receive the payload
+	 * @param ReferencingPrimPath - Path to the prim inside the referencing stage that will receive the payload
+	 * @param TargetStagePath - Path to the layer to add as payload (can be empty for internal payloads). e.g. "C:\MyFolder\Cube.usda"
+	 * @param TargetPrimPath - Path to a specific prim to target within the target layer. e.g. "/SomeOtherPrim"
+	 * @param TimeCodeOffset - Offset in timeCodes to author on the payload
+	 * @param TimeCodeScale - TimeCode scale to author on the payload
+	 * @param bUseProjectDefaultTypeHandling - If true it means we will use the ReferencerTypeHandling set on the USD project settings.
+	 *                                         If false it means we will use the provided ReferencerSchemaHandling
+	 * @param ReferencerSchemaHandling - How to behave when Prim has a different type than the target prim.
+	 *                                   Passing an unset TOptional will check the project settings for the value instead.
+	 *                                   The default value is MatchReferencedType for backwards compatibility.
+	 *                                   Note: If this is set to ShowDialog, the call may create an Editor transaction.
+	 */
 	UFUNCTION(BlueprintCallable, Category = "USD|Layer utils")
-	static void AddPayload(const FString& ReferencingStagePath, const FString& ReferencingPrimPath, const FString& TargetStagePath);
+	static void AddPayload(
+		const FString& ReferencingStagePath,
+		const FString& ReferencingPrimPath,
+		const FString& TargetStagePath,
+		const FString& TargetPrimPath = TEXT(""),
+		double TimeCodeOffset = 0.0,
+		double TimeCodeScale = 1.0,
+		bool bUseProjectDefaultTypeHandling = false,
+		EReferencerTypeHandling ReferencerTypeHandling = EReferencerTypeHandling::MatchReferencedType
+	);
 
 	UFUNCTION(BlueprintCallable, Category = "USD|Prim utils")
 	static FString GetPrimPathForObject(
diff --git a/Engine/Plugins/Importers/USDImporter/Source/USDStageEditor/Private/SUSDStageTreeView.cpp b/Engine/Plugins/Importers/USDImporter/Source/USDStageEditor/Private/SUSDStageTreeView.cpp
--- a/Engine/Plugins/Importers/USDImporter/Source/USDStageEditor/Private/SUSDStageTreeView.cpp
+++ b/Engine/Plugins/Importers/USDImporter/Source/USDStageEditor/Private/SUSDStageTreeView.cpp
@@ -1207,6 +1207,7 @@ void SUsdStageTreeView::OnAddReference()
 	UsdUtils::AddReference(
 		Referencer,
 		*ReferencedLayerPath,
+		TOptional<EReferencerTypeHandling>{},	 // Explicit unset optional to check project settings
 		OptionsPtr->bUseDefaultPrim ? UE::FSdfPath{} : UE::FSdfPath{*OptionsPtr->TargetPrimPath},
 		OptionsPtr->TimeCodeOffset,
 		OptionsPtr->TimeCodeScale
@@ -1287,6 +1288,7 @@ void SUsdStageTreeView::OnAddPayload()
 	UsdUtils::AddPayload(
 		Referencer,
 		*ReferencedLayerPath,
+		TOptional<EReferencerTypeHandling>{},	 // Explicit unset optional to check project settings
 		OptionsPtr->bUseDefaultPrim ? UE::FSdfPath{} : UE::FSdfPath{*OptionsPtr->TargetPrimPath},
 		OptionsPtr->TimeCodeOffset,
 		OptionsPtr->TimeCodeScale
diff --git a/Engine/Plugins/Importers/USDImporter/Source/USDStageEditorViewModels/Private/USDPrimViewModel.cpp b/Engine/Plugins/Importers/USDImporter/Source/USDStageEditorViewModels/Private/USDPrimViewModel.cpp
--- a/Engine/Plugins/Importers/USDImporter/Source/USDStageEditorViewModels/Private/USDPrimViewModel.cpp
+++ b/Engine/Plugins/Importers/USDImporter/Source/USDStageEditorViewModels/Private/USDPrimViewModel.cpp
@@ -2,6 +2,7 @@
 
 #include "USDPrimViewModel.h"
 
+#include "UnrealUSDWrapper.h"
 #include "USDConversionUtils.h"
 #include "USDIntegrationUtils.h"
 #include "USDMemory.h"
@@ -421,8 +422,16 @@ void FUsdPrimViewModel::ClearReferences()
 
 	FScopedUsdAllocs UsdAllocs;
 
-	pxr::UsdReferences References = pxr::UsdPrim(UsdPrim).GetReferences();
+	pxr::UsdPrim Prim{UsdPrim};
+
+	pxr::UsdReferences References = Prim.GetReferences();
 	References.ClearReferences();
+
+	// Set it back to Xform instead of fully typeless so that we at least get an actor / component in the USD Stage Editor
+	if (!Prim.HasAuthoredTypeName())
+	{
+		Prim.SetTypeName(UnrealIdentifiers::Xform);
+	}
 #endif	  // #if USE_USD_SDK
 }
 
@@ -436,8 +445,16 @@ void FUsdPrimViewModel::ClearPayloads()
 
 	FScopedUsdAllocs UsdAllocs;
 
-	pxr::UsdPayloads Payloads = pxr::UsdPrim(UsdPrim).GetPayloads();
+	pxr::UsdPrim Prim{UsdPrim};
+
+	pxr::UsdPayloads Payloads = Prim.GetPayloads();
 	Payloads.ClearPayloads();
+
+	// Set it back to Xform instead of fully typeless so that we at least get an actor / component in the USD Stage Editor
+	if (!Prim.HasAuthoredTypeName())
+	{
+		Prim.SetTypeName(UnrealIdentifiers::Xform);
+	}
 #endif	  // #if USE_USD_SDK
 }
 
diff --git a/Engine/Plugins/Runtime/USDCore/Source/USDClasses/Public/USDProjectSettings.h b/Engine/Plugins/Runtime/USDCore/Source/USDClasses/Public/USDProjectSettings.h
--- a/Engine/Plugins/Runtime/USDCore/Source/USDClasses/Public/USDProjectSettings.h
+++ b/Engine/Plugins/Runtime/USDCore/Source/USDClasses/Public/USDProjectSettings.h
@@ -24,6 +24,15 @@ enum class EUsdEditInInstanceBehavior : uint8
 	ShowPrompt
 };
 
+UENUM(Blueprinttype)
+enum class EReferencerTypeHandling : uint8
+{
+	Ignore,
+	MatchReferencedType,
+	ClearReferencerType,
+	ShowPrompt,
+};
+
 // USDCore and defaultconfig here so this ends up at DefaultUSDCore.ini in the editor, and is sent to the
 // packaged game as well
 
@@ -99,6 +108,10 @@ public:
 	UPROPERTY(config, EditAnywhere, Category = "USD|Dialogs")
 	EUsdEditInInstanceBehavior EditInInstanceableBehavior = EUsdEditInInstanceBehavior::ShowPrompt;
 
+	// How to behave when authoring a reference or payload to a prim whose type name differs from the one of the referencer prim
+	UPROPERTY(config, EditAnywhere, Category = "USD|Dialogs")
+	EReferencerTypeHandling ReferencerTypeHandling = EReferencerTypeHandling::ClearReferencerType;
+
 	// Whether to show a warning whenever the "Duplicate All Local Layer Specs" option is picked, and the duplicated
 	// prim has some specs outside the local layer stack that will not be duplicated.
 	UPROPERTY(config, EditAnywhere, Category = "USD|Dialogs")
@@ -209,7 +222,8 @@ public:
 		meta = (AllowedClasses = "/Script/Engine.MaterialInterface")
 	)
 	FSoftObjectPath ReferencePreviewSurfaceTranslucentTwoSidedMaterial = FSoftObjectPath{
-		TEXT("/USDCore/Materials/UsdPreviewSurfaceTranslucentTwoSided.UsdPreviewSurfaceTranslucentTwoSided")};
+		TEXT("/USDCore/Materials/UsdPreviewSurfaceTranslucentTwoSided.UsdPreviewSurfaceTranslucentTwoSided")
+	};
 
 	/**
 	 * What material to use as reference material when creating material instances from USD materials.
@@ -238,7 +252,8 @@ public:
 		meta = (AllowedClasses = "/Script/Engine.MaterialInterface")
 	)
 	FSoftObjectPath ReferencePreviewSurfaceTranslucentVTMaterial = FSoftObjectPath{
-		TEXT("/USDCore/Materials/UsdPreviewSurfaceTranslucentVT.UsdPreviewSurfaceTranslucentVT")};
+		TEXT("/USDCore/Materials/UsdPreviewSurfaceTranslucentVT.UsdPreviewSurfaceTranslucentVT")
+	};
 
 	/**
 	 * What material to use as reference material when creating material instances from USD materials.
@@ -268,7 +283,8 @@ public:
 		meta = (AllowedClasses = "/Script/Engine.MaterialInterface")
 	)
 	FSoftObjectPath ReferencePreviewSurfaceTranslucentTwoSidedVTMaterial = FSoftObjectPath{
-		TEXT("/USDCore/Materials/UsdPreviewSurfaceTranslucentTwoSidedVT.UsdPreviewSurfaceTranslucentTwoSidedVT")};
+		TEXT("/USDCore/Materials/UsdPreviewSurfaceTranslucentTwoSidedVT.UsdPreviewSurfaceTranslucentTwoSidedVT")
+	};
 
 	/**
 	 * What material to use as reference material when creating material instances from USD materials.
@@ -306,7 +322,8 @@ public:
 	 */
 	UPROPERTY(config, EditAnywhere, Category = "USD|Reference Materials|DisplayColor", meta = (AllowedClasses = "/Script/Engine.MaterialInterface"))
 	FSoftObjectPath ReferenceDisplayColorAndOpacityTwoSidedMaterial = FSoftObjectPath{
-		TEXT("/USDCore/Materials/DisplayColorAndOpacityTwoSided.DisplayColorAndOpacityTwoSided")};
+		TEXT("/USDCore/Materials/DisplayColorAndOpacityTwoSided.DisplayColorAndOpacityTwoSided")
+	};
 
 public:
 #if WITH_EDITOR
diff --git a/Engine/Plugins/Runtime/USDCore/Source/USDUtilities/Private/USDConversionUtils.cpp b/Engine/Plugins/Runtime/USDCore/Source/USDUtilities/Private/USDConversionUtils.cpp
--- a/Engine/Plugins/Runtime/USDCore/Source/USDUtilities/Private/USDConversionUtils.cpp
+++ b/Engine/Plugins/Runtime/USDCore/Source/USDUtilities/Private/USDConversionUtils.cpp
@@ -58,7 +58,13 @@
 #include "PhysicsEngine/PhysicsAsset.h"
 #include "Sound/AmbientSound.h"
 #include "SparseVolumeTexture/SparseVolumeTexture.h"
+#include "Widgets/Input/SButton.h"
+#include "Widgets/Layout/SBorder.h"
 #include "Widgets/Notifications/SNotificationList.h"
+#include "Widgets/Text/STextBlock.h"
+#if WITH_EDITOR
+#include "Editor.h"
+#endif	  // WITH_EDITOR
 
 #if USE_USD_SDK
 #include "USDIncludesStart.h"
@@ -648,7 +654,12 @@ FString UsdUtils::GetSchemaNameForComponent(const USceneComponent& Component)
 	return TEXT("Xform");
 }
 
-FString UsdUtils::GetPrimPathForObject(const UObject* ActorOrComponent, const FString& ParentPrimPath, bool bUseActorFolders, const FString& RootPrimName)
+FString UsdUtils::GetPrimPathForObject(
+	const UObject* ActorOrComponent,
+	const FString& ParentPrimPath,
+	bool bUseActorFolders,
+	const FString& RootPrimName
+)
 {
 	if (!ActorOrComponent)
 	{
@@ -2538,17 +2549,273 @@ FString UsdUtils::GetAssetHashPrefix(const pxr::UsdPrim& PrimForAsset, bool bSha
 namespace UE::UsdConversionUtils::Private
 {
 #if USE_USD_SDK
-	bool PrepareToAddReferenceOrPayload(
+	void HandleTypeNameAndAddReference(
+		const pxr::UsdPrim& ReferencerPrim,
+		pxr::SdfPrimSpecHandle TargetPrimSpec,
+		TOptional<EReferencerTypeHandling> ReferencerTypeHandling,
+		TFunction<void()> AddReferenceOrPayloadLambda
+	)
+	{
+		if (!ReferencerPrim || !TargetPrimSpec || !AddReferenceOrPayloadLambda)
+		{
+			return;
+		}
+
+		FScopedUsdAllocs Allocs;
+
+		pxr::TfToken ReferencerTypeName = ReferencerPrim.GetTypeName();
+		pxr::TfToken TargetTypeName = TargetPrimSpec->GetTypeName();
+
+		// Check if we need to do anything special
+		bool bNeedHandling = false;
+		bool bUnknownTargetType = false;
+		if (!ReferencerTypeName.IsEmpty() && !TargetTypeName.IsEmpty())
+		{
+			pxr::TfType TargetPrimType = pxr::UsdSchemaRegistry::GetTypeFromName(TargetTypeName);
+			if (TargetPrimType.IsUnknown())
+			{
+				bNeedHandling = true;
+				bUnknownTargetType = true;
+			}
+			else if (!ReferencerPrim.IsA(TargetPrimType))
+			{
+				bNeedHandling = true;
+			}
+		}
+		if (!bNeedHandling)
+		{
+			// The schemas already match just fine
+			AddReferenceOrPayloadLambda();
+			return;
+		}
+
+		// Get what we actually need to do
+		EReferencerTypeHandling Handling = EReferencerTypeHandling::MatchReferencedType;
+		if (ReferencerTypeHandling.IsSet())
+		{
+			Handling = ReferencerTypeHandling.GetValue();
+		}
+		else if (const UUsdProjectSettings* Settings = GetDefault<UUsdProjectSettings>())
+		{
+			Handling = Settings->ReferencerTypeHandling;
+		}
+
+		// Show the prompt and update 'Handling' to something else if we can
+		if (Handling == EReferencerTypeHandling::ShowPrompt)
+		{
+#if WITH_EDITOR
+			UUsdProjectSettings* Settings = GetMutableDefault<UUsdProjectSettings>();
+			if (!Settings)
+			{
+				return;
+			}
+
+			const FText Title = LOCTEXT("MismatchedTypeNamesText", "USD: Mismatched type names");
+			const FText SubText = FText::Format(
+				LOCTEXT(
+					"MismatchedTypeNamesSubText",
+					"Tried to add a reference or payload from prim '{0}' with type '{1}', to target prim '{2}' with type '{3}'.\n\nSince these types are not identical, it is possible that the composed prim will not have the intended behaviour.\n\nHow do you wish to proceed?"
+				),
+				FText::FromString(UsdToUnreal::ConvertPath(ReferencerPrim.GetPrimPath())),
+				FText::FromString(UsdToUnreal::ConvertToken(ReferencerTypeName)),
+				FText::FromString(UsdToUnreal::ConvertPath(TargetPrimSpec->GetPath())),
+				FText::FromString(UsdToUnreal::ConvertToken(TargetTypeName))
+			);
+			UE_LOG(LogUsd, Warning, TEXT("%s"), *SubText.ToString().Replace(TEXT("\n\n"), TEXT(" ")));
+
+			FScopedUnrealAllocs UEAllocs;
+
+			bool bDontPromptAgain = false;
+			TOptional<EReferencerTypeHandling> ChosenHandling;
+
+			// clang-format off
+			TSharedRef<SWindow> Window = SNew(SWindow)
+				.Title(Title)
+				.ClientSize(FVector2D{600, 200})
+				.SizingRule(ESizingRule::FixedSize);
+
+			TSharedPtr< SHorizontalBox > ButtonsBox = SNew(SHorizontalBox)
+				+SHorizontalBox::Slot()
+				.FillWidth(1.0)
+				.HAlign(HAlign_Left)
+				.Padding(2, 0)
+				[
+					SNew(SCheckBox)
+					.IsChecked(bDontPromptAgain) // If we're showing the prompt we know we have this unchecked
+					.OnCheckStateChanged_Lambda([&bDontPromptAgain](ECheckBoxState NewState)
+					{
+						bDontPromptAgain = NewState == ECheckBoxState::Checked;
+					})
+					[
+						SNew(STextBlock)
+						.Text(LOCTEXT("DontPromptAgainText", "Don't prompt again"))
+						.AutoWrapText(true)
+					]
+				]
+				+SHorizontalBox::Slot()
+				.AutoWidth()
+				.Padding(2, 0)
+				[
+					SNew(SButton)
+					.ButtonStyle(&FAppStyle::Get(), "PrimaryButton")
+					.TextStyle(&FAppStyle::Get(), "PrimaryButtonText")
+					.Text(LOCTEXT("ClearReferencerTypeText", "Clear referencer"))
+					.ToolTipText(LOCTEXT("ClearReferencerTypeSubText", "Remove the authored type opinion from the referencer prim, turning it into a typeless prim. This will let the type name opinion from the target prim dictate the type of the composed prim."))
+					.HAlign(HAlign_Right)
+					.VAlign(VAlign_Center)
+					.OnClicked_Lambda([&ChosenHandling, &Window]() -> FReply
+					{
+						ChosenHandling = EReferencerTypeHandling::ClearReferencerType;
+						Window->RequestDestroyWindow();
+						return FReply::Handled();
+					})
+				]
+				+SHorizontalBox::Slot()
+				.AutoWidth()
+				.Padding(2, 0)
+				[
+					SNew(SButton)
+					.ButtonStyle(&FAppStyle::Get(), "Button")
+					.TextStyle(&FAppStyle::Get(), "ButtonText")
+					.Text(LOCTEXT("MatchReferencedTypeText", "Match referenced"))
+					.ToolTipText(LOCTEXT("MatchReferencedTypeSubText", "Force the referencer prim to have the same type as the referenced prim."))
+					.HAlign(HAlign_Right)
+					.VAlign(VAlign_Center)
+					.OnClicked_Lambda([&ChosenHandling, &Window]() -> FReply
+					{
+						ChosenHandling = EReferencerTypeHandling::MatchReferencedType;
+						Window->RequestDestroyWindow();
+						return FReply::Handled();
+					})
+				]
+				+SHorizontalBox::Slot()
+				.AutoWidth()
+				.Padding(2, 0)
+				[
+					SNew(SButton)
+					.ButtonStyle(&FAppStyle::Get(), "Button")
+					.TextStyle(&FAppStyle::Get(), "ButtonText")
+					.Text(LOCTEXT("IgnoreText", "Ignore"))
+					.ToolTipText(LOCTEXT("IgnoreSubText", "Add the reference or payload anyway, ignoring the difference in type names"))
+					.HAlign(HAlign_Right)
+					.VAlign(VAlign_Center)
+					.OnClicked_Lambda([&ChosenHandling, &Window]() -> FReply
+					{
+						ChosenHandling = EReferencerTypeHandling::Ignore;
+						Window->RequestDestroyWindow();
+						return FReply::Handled();
+					})
+				];
+
+			TSharedRef<SBorder> Border = SNew(SBorder)
+			.BorderImage(FAppStyle::Get().GetBrush("Brushes.Panel"))
+			.Padding(FMargin(16))
+			.HAlign(HAlign_Fill)
+			.VAlign(VAlign_Fill)
+			[
+				SNew(SVerticalBox)
+				+SVerticalBox::Slot()
+				.FillHeight(1.0f)
+				.Padding(0.0f,0.0f,0.0f,8.0f)
+				[
+					SNew(STextBlock)
+					.Text(SubText)
+					.AutoWrapText(true)
+				]
+				+SVerticalBox::Slot()
+				.AutoHeight()
+				.Padding(0.0f, 16.0f, 0.0f, 0.0f)
+				.HAlign(HAlign_Fill)
+				.VAlign(VAlign_Bottom)
+				[
+					ButtonsBox.ToSharedRef()
+				]
+			];
+			// clang-format on
+
+			Window->SetContent(Border);
+			Window->SetWidgetToFocusOnActivate(Border);
+
+			GEditor->EditorAddModalWindow(Window);
+
+			if (ChosenHandling.IsSet())
+			{
+				Handling = ChosenHandling.GetValue();
+				ensure(Handling != EReferencerTypeHandling::ShowPrompt);	// We only set chosen handling when pressing a button
+
+				if (bDontPromptAgain)
+				{
+					Settings->ReferencerTypeHandling = ChosenHandling.GetValue();
+					Settings->SaveConfig();
+				}
+			}
+#else
+			Handling = EReferencerTypeHandling::ClearReferencerType;
+#endif	  // WITH EDITOR
+		}
+
+		switch (Handling)
+		{
+			default:
+			case EReferencerTypeHandling::Ignore:
+			{
+				AddReferenceOrPayloadLambda();
+				break;
+			}
+			case EReferencerTypeHandling::MatchReferencedType:
+			{
+				pxr::SdfChangeBlock Block;
+
+				if (bUnknownTargetType)
+				{
+					UE_LOG(
+						LogUsd,
+						Warning,
+						TEXT(
+							"Failed to match the referenced type when adding a reference or payload to prim '%s', as the target prim spec '%s' has an unknown type '%s'! The referencer type will be cleared instead."
+						),
+						*UsdToUnreal::ConvertPath(ReferencerPrim.GetPrimPath()),
+						*UsdToUnreal::ConvertPath(TargetPrimSpec->GetPath()),
+						*UsdToUnreal::ConvertToken(TargetTypeName)
+					);
+					ReferencerPrim.ClearTypeName();
+				}
+				else
+				{
+					ReferencerPrim.SetTypeName(TargetTypeName);
+				}
+
+				AddReferenceOrPayloadLambda();
+				break;
+			}
+			case EReferencerTypeHandling::ClearReferencerType:
+			{
+				pxr::SdfChangeBlock Block;
+				ReferencerPrim.ClearTypeName();
+				AddReferenceOrPayloadLambda();
+				break;
+			}
+			case EReferencerTypeHandling::ShowPrompt:
+			{
+				// We showed the dialog but didn't choose any handling --> Do nothing
+				break;
+			}
+		}
+	}
+
+	void AddReferenceOrPayload(
+		bool bIsReference,
 		const UE::FUsdPrim& Prim,
 		const TCHAR* AbsoluteFilePath,
 		const UE::FSdfPath& TargetPrimPath,
-		std::string& OutRelativeFilePath,
-		pxr::SdfPath& OutPrimPath
+		double TimeCodeOffset,
+		double TimeCodeScale,
+		TOptional<EReferencerTypeHandling> ReferencerTypeHandling
 	)
 	{
 		if (!Prim || !AbsoluteFilePath)
 		{
-			return false;
+			return;
 		}
 
 		FScopedUsdAllocs UsdAllocs;
@@ -2558,7 +2825,7 @@ namespace UE::UsdConversionUtils::Private
 		pxr::UsdStageRefPtr UsdStage = UsdPrim.GetStage();
 		if (!UsdStage)
 		{
-			return false;
+			return;
 		}
 
 		// Turn our layer path into a relative one
@@ -2592,7 +2859,7 @@ namespace UE::UsdConversionUtils::Private
 		}
 		if (!TargetLayer)
 		{
-			return false;
+			return;
 		}
 
 		// Get the target prim spec we want to reference
@@ -2602,28 +2869,31 @@ namespace UE::UsdConversionUtils::Private
 			TargetPrimSpec = TargetLayer->GetPrimAtPath(pxr::SdfPath::AbsoluteRootPath().AppendChild(TargetLayer->GetDefaultPrim()));
 		}
 
-		// Update UsdPrim's type so that it can handle the reference and be parsed by the proper translator
-		if (!UsdPrim.GetTypeName().IsEmpty() && TargetPrimSpec)
-		{
-			pxr::TfToken TargetTypeName = TargetPrimSpec->GetTypeName();
-			pxr::TfType TargetPrimType = pxr::UsdSchemaRegistry::GetTypeFromName(TargetTypeName);
+		// We want to output no path for the prim if we received it as such, even if we already know what the path to the
+		// default prim is, so that the authored reference doesn't actually specify any prim name and just refers to the
+		// default prim by default. Otherwise if the default prim of the layer changed, we wouldn't update to the new prim
+		pxr::SdfPath FinalPrimPath = TargetPrimSpec && !TargetPrimPath.IsEmpty() ? TargetPrimSpec->GetPath() : pxr::SdfPath{};
+		std::string RelativeLayerPath = bIsInternalReference ? std::string() : UnrealToUsd::ConvertString(*RelativePath).Get();
 
-			if (TargetPrimType.IsUnknown())
+		TFunction<void()> AddReferenceOrPayloadLambda = nullptr;
+		if (bIsReference)
+		{
+			AddReferenceOrPayloadLambda = [UsdPrim, &RelativeLayerPath, &FinalPrimPath, TimeCodeOffset, TimeCodeScale]()
 			{
-				UsdPrim.ClearTypeName();
-			}
-			else if (!UsdPrim.IsA(TargetPrimType))
+				pxr::UsdReferences References = UsdPrim.GetReferences();
+				References.AddReference(RelativeLayerPath, FinalPrimPath, pxr::SdfLayerOffset{TimeCodeOffset, TimeCodeScale});
+			};
+		}
+		else	// It's a payload instead
+		{
+			AddReferenceOrPayloadLambda = [UsdPrim, &RelativeLayerPath, &FinalPrimPath, TimeCodeOffset, TimeCodeScale]()
 			{
-				UsdPrim.SetTypeName(TargetTypeName);
-			}
+				pxr::UsdPayloads Payloads = UsdPrim.GetPayloads();
+				Payloads.AddPayload(RelativeLayerPath, FinalPrimPath, pxr::SdfLayerOffset{TimeCodeOffset, TimeCodeScale});
+			};
 		}
 
-		// We want to output no path for the prim if we received it as such, even if we already know what the path to the
-		// default prim is, so that the authored reference doesn't actually specify any prim name and just refers to the
-		// default prim by default. Otherwise if the default prim of the layer changed, we wouldn't update to the new prim
-		OutPrimPath = TargetPrimSpec && !TargetPrimPath.IsEmpty() ? TargetPrimSpec->GetPath() : pxr::SdfPath{};
-		OutRelativeFilePath = bIsInternalReference ? std::string() : UnrealToUsd::ConvertString(*RelativePath).Get();
-		return true;
+		HandleTypeNameAndAddReference(Prim, TargetPrimSpec, ReferencerTypeHandling, AddReferenceOrPayloadLambda);
 	}
 #endif	  // #if USE_USD_SDK
 }	 // namespace UE::UsdConversionUtils::Private
@@ -2636,30 +2906,30 @@ void UsdUtils::AddReference(
 	double TimeCodeScale
 )
 {
-#if USE_USD_SDK
-	FScopedUsdAllocs UsdAllocs;
-
-	// Group updates or else the SetTypeName (inside PrepareToAddReferenceOrPayload) and AddReference calls below will
-	// both trigger separate resyncs of the same prim path
-	pxr::SdfChangeBlock ChangeBlock;
+	const EReferencerTypeHandling Handling = EReferencerTypeHandling::MatchReferencedType;
+	AddReference(Prim, AbsoluteFilePath, Handling, TargetPrimPath, TimeCodeOffset, TimeCodeScale);
+}
 
-	std::string RelativeLayerPath;
-	pxr::SdfPath FinalPrimPath;
-	bool bProceed = UE::UsdConversionUtils::Private::PrepareToAddReferenceOrPayload(
+void UsdUtils::AddReference(
+	UE::FUsdPrim& Prim,
+	const TCHAR* AbsoluteFilePath,
+	TOptional<EReferencerTypeHandling> ReferencerTypeHandling,
+	const UE::FSdfPath& TargetPrimPath,
+	double TimeCodeOffset,
+	double TimeCodeScale
+)
+{
+#if USE_USD_SDK
+	const bool bIsReference = true;
+	UE::UsdConversionUtils::Private::AddReferenceOrPayload(
+		bIsReference,
 		Prim,
 		AbsoluteFilePath,
 		TargetPrimPath,
-		RelativeLayerPath,
-		FinalPrimPath
+		TimeCodeOffset,
+		TimeCodeScale,
+		ReferencerTypeHandling
 	);
-
-	if (!bProceed)
-	{
-		return;
-	}
-
-	pxr::UsdReferences References = pxr::UsdPrim{Prim}.GetReferences();
-	References.AddReference(RelativeLayerPath, FinalPrimPath, pxr::SdfLayerOffset{TimeCodeOffset, TimeCodeScale});
 #endif	  // #if USE_USD_SDK
 }
 
@@ -2702,30 +2972,30 @@ void UsdUtils::AddPayload(
 	double TimeCodeScale
 )
 {
-#if USE_USD_SDK
-	FScopedUsdAllocs UsdAllocs;
-
-	// Group updates or else the SetTypeName (inside PrepareToAddReferenceOrPayload) and AddReference calls below will
-	// both trigger separate resyncs of the same prim path
-	pxr::SdfChangeBlock ChangeBlock;
+	const EReferencerTypeHandling Handling = EReferencerTypeHandling::MatchReferencedType;
+	AddPayload(Prim, AbsoluteFilePath, Handling, TargetPrimPath, TimeCodeOffset, TimeCodeScale);
+}
 
-	std::string RelativeLayerPath;
-	pxr::SdfPath FinalPrimPath;
-	bool bProceed = UE::UsdConversionUtils::Private::PrepareToAddReferenceOrPayload(
+void UsdUtils::AddPayload(
+	UE::FUsdPrim& Prim,
+	const TCHAR* AbsoluteFilePath,
+	TOptional<EReferencerTypeHandling> ReferencerTypeHandling,
+	const UE::FSdfPath& TargetPrimPath,
+	double TimeCodeOffset,
+	double TimeCodeScale
+)
+{
+#if USE_USD_SDK
+	const bool bIsReference = false;
+	UE::UsdConversionUtils::Private::AddReferenceOrPayload(
+		bIsReference,
 		Prim,
 		AbsoluteFilePath,
 		TargetPrimPath,
-		RelativeLayerPath,
-		FinalPrimPath
+		TimeCodeOffset,
+		TimeCodeScale,
+		ReferencerTypeHandling
 	);
-
-	if (!bProceed)
-	{
-		return;
-	}
-
-	pxr::UsdPayloads Payloads = pxr::UsdPrim{Prim}.GetPayloads();
-	Payloads.AddPayload(RelativeLayerPath, FinalPrimPath, pxr::SdfLayerOffset{TimeCodeOffset, TimeCodeScale});
 #endif	  // #if USE_USD_SDK
 }
 
@@ -4041,7 +4311,7 @@ void UsdUtils::ReadStageMetaData(UE::FUsdStage Stage, TMap<FString, FString>& Ou
 		OutMetaDataMap.Add(TEXT("Up Axis"), TEXT("Y"));
 	}
 
-	//DefaultPrim
+	// DefaultPrim
 	UE::FVtValue DefaultPrim;
 	if (Stage.GetMetadata(*UsdToUnreal::ConvertToken(SdfFieldKeys->DefaultPrim), DefaultPrim))
 	{
@@ -4059,27 +4329,27 @@ void UsdUtils::ReadStageMetaData(UE::FUsdStage Stage, TMap<FString, FString>& Ou
 		}
 	}
 
-	//double typed potential meta data handling:
+	// double typed potential meta data handling:
 	auto ProcessDoubleTypedMetadata = [&OutMetaDataMap, &Stage](const FString& ExtraInformationKey, const FString& TokenStringified)
+	{
+		UE::FVtValue UEValue;
+		if (Stage.GetMetadata(*TokenStringified, UEValue))
 		{
-			UE::FVtValue UEValue;
-			if (Stage.GetMetadata(*TokenStringified, UEValue))
-			{
-				VtValue UsdValue = UEValue.GetUsdValue();
+			VtValue UsdValue = UEValue.GetUsdValue();
 
-				if (UsdValue.IsHolding<double>())
-				{
-					double Value = UsdValue.UncheckedGet<double>();
-					OutMetaDataMap.Add(ExtraInformationKey, FString::SanitizeFloat(Value));
-				}
+			if (UsdValue.IsHolding<double>())
+			{
+				double Value = UsdValue.UncheckedGet<double>();
+				OutMetaDataMap.Add(ExtraInformationKey, FString::SanitizeFloat(Value));
 			}
-		};
+		}
+	};
 	ProcessDoubleTypedMetadata(TEXT("Time Codes Per Second"), UsdToUnreal::ConvertToken(SdfFieldKeys->TimeCodesPerSecond));
 	ProcessDoubleTypedMetadata(TEXT("Frames Per Second"), UsdToUnreal::ConvertToken(SdfFieldKeys->FramesPerSecond));
 	ProcessDoubleTypedMetadata(TEXT("Start Time Code"), UsdToUnreal::ConvertToken(SdfFieldKeys->StartTimeCode));
 	ProcessDoubleTypedMetadata(TEXT("End Time Code"), UsdToUnreal::ConvertToken(SdfFieldKeys->EndTimeCode));
 
-	//Documentation
+	// Documentation
 	UE::FVtValue Documentation;
 	if (Stage.GetMetadata(*UsdToUnreal::ConvertToken(SdfFieldKeys->Documentation), Documentation))
 	{
@@ -4092,7 +4362,7 @@ void UsdUtils::ReadStageMetaData(UE::FUsdStage Stage, TMap<FString, FString>& Ou
 		}
 	}
 
-	//CustomLayerData
+	// CustomLayerData
 	UE::FVtValue CustomLayerData;
 	if (Stage.GetMetadata(*UsdToUnreal::ConvertToken(SdfFieldKeys->CustomLayerData), CustomLayerData))
 	{
diff --git a/Engine/Plugins/Runtime/USDCore/Source/USDUtilities/Private/USDGeomMeshConversion.cpp b/Engine/Plugins/Runtime/USDCore/Source/USDUtilities/Private/USDGeomMeshConversion.cpp
--- a/Engine/Plugins/Runtime/USDCore/Source/USDUtilities/Private/USDGeomMeshConversion.cpp
+++ b/Engine/Plugins/Runtime/USDCore/Source/USDUtilities/Private/USDGeomMeshConversion.cpp
@@ -1304,7 +1304,6 @@ namespace UE::UsdGeomMeshConversion::Private
 
 	void CollectPrimvarData(const pxr::UsdPrim& Prim, const UsdToUnreal::FUsdMeshConversionOptions& Options, FUsdMeshData& InOutMeshData)
 	{
-		using namespace UsdToUnreal;
 		const pxr::TfToken DefaultInterpolation = pxr::UsdGeomTokens->varying;
 
 		pxr::UsdGeomPrimvarsAPI PrimvarAPI{Prim};
@@ -1312,8 +1311,7 @@ namespace UE::UsdGeomMeshConversion::Private
 		for (const pxr::UsdGeomPrimvar& Primvar : Primvars)
 		{
 			std::string PrimvarName = Primvar.GetPrimvarName().GetString();
-			if (PrimvarName == "displayColor" || PrimvarName == "displayOpacity" || PrimvarName == "st"
-				|| (Options.ImportPrimvars == FUsdMeshConversionOptions::EImportPrimvar::Bake && !Options.PrimvarNames.Contains(UTF8_TO_TCHAR(PrimvarName.c_str()))))
+			if (PrimvarName == "displayColor" || PrimvarName == "displayOpacity" || PrimvarName == "st")
 			{
 				continue;
 			}
@@ -1499,7 +1497,6 @@ namespace UE::UsdGeomMeshConversion::Private
 			}
 
 			// Custom primvars
-			if (Options.ImportPrimvars != UsdToUnreal::FUsdMeshConversionOptions::EImportPrimvar::Standard)
 			{
 				CollectPrimvarData(Prim, Options, InOutMeshData);
 			}
@@ -6525,7 +6522,10 @@ void UsdUtils::ReplaceUnrealMaterialsWithBaked(
 						}
 
 						UE::FUsdPrim UEMatPrim{MatPrim};
-						UsdUtils::AddReference(UEMatPrim, *BakedFilename);
+
+						const UE::FSdfPath TargetPrimPath;
+						const EReferencerTypeHandling Handling = EReferencerTypeHandling::ClearReferencerType;
+						UsdUtils::AddReference(UEMatPrim, *BakedFilename, Handling, TargetPrimPath);
 					}
 				}
 				// Need a MatScopePrim authored somewhere within this layer
@@ -6573,8 +6573,7 @@ void UsdUtils::ReplaceUnrealMaterialsWithBaked(
 						MatPrimScopePtr->UsedPrimNames.Add(MatPrimName);
 
 						MatPrim = StageToTraverse->DefinePrim(
-							MatPrimScopePtr->Prim.GetPath().AppendChild(UnrealToUsd::ConvertToken(*MatPrimName).Get()),
-							UnrealToUsd::ConvertToken(TEXT("Material")).Get()
+							MatPrimScopePtr->Prim.GetPath().AppendChild(UnrealToUsd::ConvertToken(*MatPrimName).Get())
 						);
 
 						// We should only keep track and reuse the material proxy prims that we create within the MatPrimScope, not
diff --git a/Engine/Plugins/Runtime/USDCore/Source/USDUtilities/Public/USDConversionUtils.h b/Engine/Plugins/Runtime/USDCore/Source/USDUtilities/Public/USDConversionUtils.h
--- a/Engine/Plugins/Runtime/USDCore/Source/USDUtilities/Public/USDConversionUtils.h
+++ b/Engine/Plugins/Runtime/USDCore/Source/USDUtilities/Public/USDConversionUtils.h
@@ -8,6 +8,7 @@
 #include "UObject/UObjectGlobals.h"
 
 #include "UnrealUSDWrapper.h"
+#include "USDProjectSettings.h"
 #if UE_ENABLE_INCLUDE_ORDER_DEPRECATED_IN_5_5
 #include "USDMemory.h"
 #endif	  // UE_ENABLE_INCLUDE_ORDER_DEPRECATED_IN_5_5
@@ -460,10 +461,29 @@ namespace UsdUtils
 	USDUTILITIES_API FString GetAssetHashPrefix(const pxr::UsdPrim& PrimForAsset, bool bShareAssetsForIdenticalPrims);
 #endif	  // #if USE_USD_SDK
 
-	/** Adds a reference on Prim to the layer at AbsoluteFilePath */
+	/**
+	 * Adds a reference on Prim to the layer at AbsoluteFilePath, optionally specifying a target prim inside that layer.
+	 * @param Prim - Referencer prim
+	 * @param AbsoluteFilePath - Path to the layer to reference (can be empty for internal references). e.g. "C:\MyFolder\Cube.usda"
+	 * @param TargetPrimPath - Path to a specific prim to reference within the target layer. e.g. "/SomeOtherPrim"
+	 * @param TimeCodeOffset - Offset in timeCodes to author on the reference
+	 * @param TimeCodeScale - TimeCode scale to author on the reference
+	 * @param ReferencerSchemaHandling - How to behave when Prim has a different type than the referenced prim.
+	 *                                   Passing an unset TOptional will check the project settings for the value instead.
+	 *                                   The default value is MatchReferencedType for backwards compatibility.
+	 *                                   Note: If this is set to ShowDialog, the call may create an Editor transaction.
+	 */
+	USDUTILITIES_API void AddReference(
+		UE::FUsdPrim& Prim,
+		const TCHAR* AbsoluteFilePath,
+		const UE::FSdfPath& TargetPrimPath = {},
+		double TimeCodeOffset = 0.0,
+		double TimeCodeScale = 1.0
+	);
 	USDUTILITIES_API void AddReference(
 		UE::FUsdPrim& Prim,
 		const TCHAR* AbsoluteFilePath,
+		TOptional<EReferencerTypeHandling> ReferencerTypeHandling,
 		const UE::FSdfPath& TargetPrimPath = {},
 		double TimeCodeOffset = 0.0,
 		double TimeCodeScale = 1.0
@@ -472,10 +492,29 @@ namespace UsdUtils
 	/** Gets the strongest direct reference on Prim with the given FileExtension. Returns true if there was one. */
 	USDUTILITIES_API bool GetReferenceFilePath(const UE::FUsdPrim& Prim, const FString& FileExtension, FString& OutReferenceFilePath);
 
-	/** Adds a payload on Prim pointing at the default prim of the layer at AbsoluteFilePath */
+	/**
+	 * Adds a payload on Prim pointing at the layer at AbsoluteFilePath, optionally specifying a target prim inside that layer.
+	 * @param Prim - Prim to receive the payload
+	 * @param AbsoluteFilePath - Path to the layer to use as payload (can be empty for internal references). e.g. "C:\MyFolder\Cube.usda"
+	 * @param TargetPrimPath - Path to a specific prim to use as payload within the target layer. e.g. "/SomeOtherPrim"
+	 * @param TimeCodeOffset - Offset in timeCodes to author on the payload
+	 * @param TimeCodeScale - TimeCode scale to author on the payload
+	 * @param ReferencerSchemaHandling - How to behave when Prim has a different type than the prim used by the payload.
+	 *                                   Passing an unset TOptional will check the project settings for the value instead.
+	 *                                   The default value is MatchReferencedType for backwards compatibility.
+	 *                                   Note: If this is set to ShowDialog, the call may create an Editor transaction.
+	 */
+	USDUTILITIES_API void AddPayload(
+		UE::FUsdPrim& Prim,
+		const TCHAR* AbsoluteFilePath,
+		const UE::FSdfPath& TargetPrimPath = {},
+		double TimeCodeOffset = 0.0,
+		double TimeCodeScale = 1.0
+	);
 	USDUTILITIES_API void AddPayload(
 		UE::FUsdPrim& Prim,
 		const TCHAR* AbsoluteFilePath,
+		TOptional<EReferencerTypeHandling> ReferencerTypeHandling,
 		const UE::FSdfPath& TargetPrimPath = {},
 		double TimeCodeOffset = 0.0,
 		double TimeCodeScale = 1.0
diff --git a/Engine/Plugins/Runtime/USDCore/Source/UnrealUSDWrapper/Private/UnrealUSDWrapper.cpp b/Engine/Plugins/Runtime/USDCore/Source/UnrealUSDWrapper/Private/UnrealUSDWrapper.cpp
--- a/Engine/Plugins/Runtime/USDCore/Source/UnrealUSDWrapper/Private/UnrealUSDWrapper.cpp
+++ b/Engine/Plugins/Runtime/USDCore/Source/UnrealUSDWrapper/Private/UnrealUSDWrapper.cpp
@@ -127,6 +127,8 @@ namespace UnrealIdentifiers
 	const TfToken Occlusion = TfToken("occlusion");
 	const TfToken Refraction = TfToken("ior");
 
+	const TfToken Xform = TfToken("Xform");
+
 	const TfToken Surface = TfToken("surface");
 	const TfToken St = TfToken("st");
 	const TfToken Varname = TfToken("varname");
diff --git a/Engine/Plugins/Runtime/USDCore/Source/UnrealUSDWrapper/Public/UnrealUSDWrapper.h b/Engine/Plugins/Runtime/USDCore/Source/UnrealUSDWrapper/Public/UnrealUSDWrapper.h
--- a/Engine/Plugins/Runtime/USDCore/Source/UnrealUSDWrapper/Public/UnrealUSDWrapper.h
+++ b/Engine/Plugins/Runtime/USDCore/Source/UnrealUSDWrapper/Public/UnrealUSDWrapper.h
@@ -455,6 +455,8 @@ namespace UnrealIdentifiers
 	extern UNREALUSDWRAPPER_API const pxr::TfToken Occlusion;
 	extern UNREALUSDWRAPPER_API const pxr::TfToken Refraction;
 
+	extern UNREALUSDWRAPPER_API const pxr::TfToken Xform;
+
 	// Tokens used mostly for shade material conversion
 	extern UNREALUSDWRAPPER_API const pxr::TfToken Surface;
 	extern UNREALUSDWRAPPER_API const pxr::TfToken St;
