USD: Fixed geometry cache export not using the correct endTimeCode and frame rate.

diff --git a/Engine/Plugins/Importers/USDImporter/Source/USDSchemas/Private/USDGeometryCacheTranslator.cpp b/Engine/Plugins/Importers/USDImporter/Source/USDSchemas/Private/USDGeometryCacheTranslator.cpp
--- a/Engine/Plugins/Importers/USDImporter/Source/USDSchemas/Private/USDGeometryCacheTranslator.cpp
+++ b/Engine/Plugins/Importers/USDImporter/Source/USDSchemas/Private/USDGeometryCacheTranslator.cpp
@@ -471,7 +471,7 @@ namespace UsdGeometryCacheTranslatorImpl
 				// StartOffsetTime is the offset applied to the GeometryCache section on the sequencer track, so not relevant when importing
 				StartOffsetTime = static_cast<float>(Args.StartFrame) / Args.FramesPerSecond;
 			}
-			GeometryCache->SetFrameStartEnd(Args.StartFrame, Args.EndFrame);
+			GeometryCache->SetFrameStartEnd(Args.StartFrame, Args.EndFrame - 1);
 
 			// Create a track for each mesh to be processed and add it to the GeometryCache
 			for (int32 Index = 0; Index < MeshPaths.Num(); ++Index)
diff --git a/Engine/Plugins/Runtime/USDCore/Source/USDUtilities/Private/USDGeomMeshConversion.cpp b/Engine/Plugins/Runtime/USDCore/Source/USDUtilities/Private/USDGeomMeshConversion.cpp
--- a/Engine/Plugins/Runtime/USDCore/Source/USDUtilities/Private/USDGeomMeshConversion.cpp
+++ b/Engine/Plugins/Runtime/USDCore/Source/USDUtilities/Private/USDGeomMeshConversion.cpp
@@ -4841,8 +4841,8 @@ namespace UE::UsdGeometryCacheConversion::Private
 			: SlotNames(GeometryCache.MaterialSlotNames)
 		{
 			// The GeometryCache's EndFrame is exclusive since it's there to allow frame interpolation past the real last frame
-			InclusiveEndFrame = FMath::Max(GeometryCache.GetEndFrame() - 1, GeometryCache.GetStartFrame() + 1);
-			FrameRate = (InclusiveEndFrame - GeometryCache.GetStartFrame()) / GeometryCache.CalculateDuration();
+			InclusiveEndFrame = FMath::Max(GeometryCache.GetEndFrame(), GeometryCache.GetStartFrame() + 1);
+			FrameRate = FMath::RoundToFloat((GeometryCache.GetEndFrame() - GeometryCache.GetStartFrame()) / GeometryCache.CalculateDuration());
 		}
 
 		FGeometryCacheExportContext(const TArray<FName>& InSlotNames)
