[USD] Restore support for custom morph target / blend shape normals

#rnx

diff --git a/Engine/Plugins/Importers/USDImporter/Source/USDSchemas/Private/USDSkelRootTranslator.cpp b/Engine/Plugins/Importers/USDImporter/Source/USDSchemas/Private/USDSkelRootTranslator.cpp
--- a/Engine/Plugins/Importers/USDImporter/Source/USDSchemas/Private/USDSkelRootTranslator.cpp
+++ b/Engine/Plugins/Importers/USDImporter/Source/USDSchemas/Private/USDSkelRootTranslator.cpp
@@ -566,7 +566,8 @@ namespace UsdSkelRootTranslatorImpl
 							LODIndex,
 							NumPointsBeforeThisMesh,
 							InOutUsedMorphTargetNames,
-							*OutBlendShapes
+							*OutBlendShapes,
+							Options
 						);
 					}
 				}
diff --git a/Engine/Plugins/Importers/USDImporter/Source/USDSchemas/Private/USDSkelSkeletonTranslator.cpp b/Engine/Plugins/Importers/USDImporter/Source/USDSchemas/Private/USDSkelSkeletonTranslator.cpp
--- a/Engine/Plugins/Importers/USDImporter/Source/USDSchemas/Private/USDSkelSkeletonTranslator.cpp
+++ b/Engine/Plugins/Importers/USDImporter/Source/USDSchemas/Private/USDSkelSkeletonTranslator.cpp
@@ -616,7 +616,8 @@ namespace UsdSkelSkeletonTranslatorImpl
 							LODIndex,
 							NumPointsBeforeThisMesh,
 							InOutUsedMorphTargetNames,
-							*OutBlendShapes
+							*OutBlendShapes,
+							Options
 						);
 					}
 				}
diff --git a/Engine/Plugins/Interchange/Extensions/OpenUSD/Source/Import/Private/InterchangeUsdTranslator.cpp b/Engine/Plugins/Interchange/Extensions/OpenUSD/Source/Import/Private/InterchangeUsdTranslator.cpp
--- a/Engine/Plugins/Interchange/Extensions/OpenUSD/Source/Import/Private/InterchangeUsdTranslator.cpp
+++ b/Engine/Plugins/Interchange/Extensions/OpenUSD/Source/Import/Private/InterchangeUsdTranslator.cpp
@@ -3809,7 +3809,7 @@ namespace UE::InterchangeUsdTranslator::Private
 		}
 
 		const float Weight = 1.0f;
-		return UsdUtils::ApplyBlendShape(OutMeshDescription, BlendShape.GetPrim(), Weight, InbetweenName);
+		return UsdUtils::ApplyBlendShape(OutMeshDescription, BlendShape.GetPrim(), Options.AdditionalTransform, Weight, InbetweenName);
 	}
 
 	// Volume animations need a special property reader because in USD it's an animation of the file path attribute within
diff --git a/Engine/Plugins/Runtime/USDCore/Source/USDUtilities/Private/USDSkeletalDataConversion.cpp b/Engine/Plugins/Runtime/USDCore/Source/USDUtilities/Private/USDSkeletalDataConversion.cpp
--- a/Engine/Plugins/Runtime/USDCore/Source/USDUtilities/Private/USDSkeletalDataConversion.cpp
+++ b/Engine/Plugins/Runtime/USDCore/Source/USDUtilities/Private/USDSkeletalDataConversion.cpp
@@ -33,6 +33,7 @@
 #include "Components/SkeletalMeshComponent.h"
 #include "ControlRig.h"
 #include "Engine/SkeletalMesh.h"
+#include "Engine/SkeletalMeshLODSettings.h"
 #include "Engine/SkinnedAssetCommon.h"
 #include "Evaluation/MovieSceneSequenceTransform.h"
 #include "IMovieScenePlayer.h"
@@ -290,7 +291,8 @@ namespace SkelDataConversionImpl
 		const FUsdStageInfo& StageInfo,
 		uint32 PointIndexOffset,
 		int32 LODIndex,
-		UsdUtils::FUsdBlendShape& OutBlendShape
+		UsdUtils::FUsdBlendShape& OutBlendShape,
+		const UsdToUnreal::FUsdMeshConversionOptions& Options
 	)
 	{
 		uint32 NumOffsets = PointOffsets.size();
@@ -359,11 +361,14 @@ namespace SkelDataConversionImpl
 			}
 		}
 
+		FMatrix TotalMatrix = Options.AdditionalTransform.ToMatrixWithScale();
+		FMatrix TotalMatrixForNormal = TotalMatrix.Inverse().GetTransposed();
+
 		OutBlendShape.Vertices.SetNumUninitialized(NumOffsets);
 		for (uint32 OffsetIndex = 0; OffsetIndex < NumOffsets; ++OffsetIndex)
 		{
 			const FVector UEOffset = UsdToUnreal::ConvertVector(StageInfo, PointOffsets[OffsetIndex]);
-			const FVector UENormal = OutBlendShape.bHasAuthoredTangents ? UsdToUnreal::ConvertVector(StageInfo, NormalOffsets[OffsetIndex])
+			const FVector UENormal = OutBlendShape.bHasAuthoredTangents ? FVector(TotalMatrixForNormal.TransformVector(UsdToUnreal::ConvertVector(StageInfo, NormalOffsets[OffsetIndex])).GetSafeNormal())
 																		: FVector(0, 0, 0);
 
 			FMorphTargetDelta& ModifiedVertex = OutBlendShape.Vertices[OffsetIndex];
@@ -1443,14 +1448,31 @@ bool UsdToUnreal::ConvertSkinnedMesh(
 	uint32 NumVertexInstances = static_cast<uint32>(OriginalFaceIndices.size());
 
 	// Normals
-	VtArray<GfVec3f> Normals;
+	TArray<FVector3f> Normals;//Transformed
+	VtArray<GfVec3f> USDNormals;
 	UsdAttribute NormalsAttribute = UsdMesh.GetNormalsAttr();
+	bool bUsingVertexInstanceNormals = false;
 	if (NormalsAttribute)
 	{
-		if (NormalsAttribute.Get(&Normals, UsdTimeCode::Default()) && Normals.size() > 0)
+		FMatrix TotalMatrix = CommonOptions.AdditionalTransform.ToMatrixWithScale();
+		FMatrix TotalMatrixForNormal = TotalMatrix.Inverse().GetTransposed();
+
+		if (NormalsAttribute.Get(&USDNormals, UsdTimeCode::Default()) && USDNormals.size() > 0)
 		{
 			// Like for points, we need to ensure these normals are in the same coordinate space of the skeleton
-			SkinningQuery.ComputeSkinnedNormals(MeshToSkeletonRestPose, &Normals, UsdTimeCode::EarliestTime());
+			SkinningQuery.ComputeSkinnedNormals(MeshToSkeletonRestPose, &USDNormals, UsdTimeCode::EarliestTime());
+
+			if (USDNormals.size() == NumPoints || USDNormals.size() == NumVertexInstances)
+			{
+				bUsingVertexInstanceNormals = USDNormals.size() == NumVertexInstances;
+				Normals.Reserve(bUsingVertexInstanceNormals ? NumVertexInstances : NumPoints);
+				SkelMeshImportData.bHasNormals = true;
+				for (GfVec3f& Normal : USDNormals)
+				{
+					FVector TransformedNormal = TotalMatrixForNormal.TransformVector(UsdToUnreal::ConvertVector(StageInfo, Normal)).GetSafeNormal();
+					Normals.Add((FVector3f)TransformedNormal);
+				}
+			}
 		}
 	}
 
@@ -1829,8 +1851,6 @@ bool UsdToUnreal::ConvertSkinnedMesh(
 			// Set the face smoothing by default. It could be any number, but not zero
 			Triangle.SmoothingGroups = 255;
 
-			// #ueent_todo: Convert normals to TangentZ (TangentX/TangentY are tangents/binormals)
-
 			Triangle.MatIndex = RealMaterialIndex;
 			Triangle.AuxMatIndex = 0;
 
@@ -1923,6 +1943,19 @@ bool UsdToUnreal::ConvertSkinnedMesh(
 				Triangle.TangentZ[FinalCornerIndex] = FVector3f::ZeroVector;
 
 				Triangle.WedgeIndex[FinalCornerIndex] = WedgeIndex;
+
+				if (SkelMeshImportData.bHasNormals)
+				{
+					if (bUsingVertexInstanceNormals)
+					{
+						Triangle.TangentZ[FinalCornerIndex] = Normals[OriginalVertexInstanceIndex];
+					}
+					else
+					{
+						Triangle.TangentZ[FinalCornerIndex] = Normals[SkelMeshWedge.VertexIndex];
+					}
+					
+				}
 			}
 		}
 
@@ -2540,10 +2573,11 @@ bool UsdToUnreal::ConvertBlendShape(
 	const FUsdStageInfo& StageInfo,
 	uint32 PointIndexOffset,
 	TSet<FString>& UsedMorphTargetNames,
-	UsdUtils::FBlendShapeMap& OutBlendShapes
+	UsdUtils::FBlendShapeMap& OutBlendShapes,
+	const FUsdMeshConversionOptions& Options
 )
 {
-	return ConvertBlendShape(UsdBlendShape, StageInfo, 0, PointIndexOffset, UsedMorphTargetNames, OutBlendShapes);
+	return ConvertBlendShape(UsdBlendShape, StageInfo, 0, PointIndexOffset, UsedMorphTargetNames, OutBlendShapes, Options);
 }
 
 bool UsdToUnreal::ConvertBlendShape(
@@ -2552,7 +2586,8 @@ bool UsdToUnreal::ConvertBlendShape(
 	int32 LODIndex,
 	uint32 PointIndexOffset,
 	TSet<FString>& UsedMorphTargetNames,
-	UsdUtils::FBlendShapeMap& OutBlendShapes
+	UsdUtils::FBlendShapeMap& OutBlendShapes,
+	const UsdToUnreal::FUsdMeshConversionOptions& Options
 )
 {
 	TRACE_CPUPROFILER_EVENT_SCOPE(UsdToUnreal::ConvertBlendShape);
@@ -2593,7 +2628,7 @@ bool UsdToUnreal::ConvertBlendShape(
 
 	UsdUtils::FUsdBlendShape PrimaryBlendShape;
 	if (!SkelDataConversionImpl::
-			CreateUsdBlendShape(PrimaryName, Offsets, Normals, PointIndices, StageInfo, PointIndexOffset, LODIndex, PrimaryBlendShape))
+			CreateUsdBlendShape(PrimaryName, Offsets, Normals, PointIndices, StageInfo, PointIndexOffset, LODIndex, PrimaryBlendShape, Options))
 	{
 		return false;
 	}
@@ -2644,7 +2679,8 @@ bool UsdToUnreal::ConvertBlendShape(
 				StageInfo,
 				PointIndexOffset,
 				LODIndex,
-				InbetweenShape
+				InbetweenShape,
+				Options
 			))
 		{
 			continue;
@@ -3285,6 +3321,7 @@ bool UsdUtils::GetSkelQueries(
 bool UsdUtils::ApplyBlendShape(
 	FMeshDescription& InOutMeshDescription,
 	const pxr::UsdPrim& InBlendShapePrim,
+	const FTransform& AdditionalTransform,
 	float Weight,
 	const FString InInbetweenName
 )
@@ -3384,9 +3421,12 @@ bool UsdUtils::ApplyBlendShape(
 		// Normal offsets
 		if (NormalOffsets.size() == (size_t)MeshPositions.GetNumElements())
 		{
+			FMatrix TotalMatrix = AdditionalTransform.ToMatrixWithScale();
+			FMatrix TotalMatrixForNormal = TotalMatrix.Inverse().GetTransposed();
+
 			for (uint32 VertexIndex = 0; VertexIndex < NormalOffsets.size(); ++VertexIndex)
 			{
-				const FVector UENormal = UsdToUnreal::ConvertVector(StageInfo, NormalOffsets[VertexIndex]).GetSafeNormal();
+				const FVector UENormal = TotalMatrixForNormal.TransformVector(UsdToUnreal::ConvertVector(StageInfo, NormalOffsets[VertexIndex])).GetSafeNormal();
 
 				TArrayView<const FVertexInstanceID> Instances = InOutMeshDescription.GetVertexVertexInstanceIDs(VertexIndex);
 				for (FVertexInstanceID InstanceID : Instances)
@@ -3448,12 +3488,15 @@ bool UsdUtils::ApplyBlendShape(
 		// Normal offsets
 		if (PointIndices.size() == NormalOffsets.size())
 		{
+			FMatrix TotalMatrix = AdditionalTransform.ToMatrixWithScale();
+			FMatrix TotalMatrixForNormal = TotalMatrix.Inverse().GetTransposed();
+
 			for (uint32 NormalIndex = 0; NormalIndex < NormalOffsets.size(); ++NormalIndex)
 			{
 				int TargetPointIndex = PointIndices[NormalIndex];
 				if (TargetPointIndex >= 0 && TargetPointIndex < MeshPositions.GetNumElements())
 				{
-					const FVector UENormal = UsdToUnreal::ConvertVector(StageInfo, NormalOffsets[NormalIndex]).GetSafeNormal();
+					const FVector UENormal = TotalMatrixForNormal.TransformVector(UsdToUnreal::ConvertVector(StageInfo, NormalOffsets[NormalIndex])).GetSafeNormal();
 
 					TArrayView<const FVertexInstanceID> Instances = InOutMeshDescription.GetVertexVertexInstanceIDs(TargetPointIndex);
 					for (FVertexInstanceID InstanceID : Instances)
diff --git a/Engine/Plugins/Runtime/USDCore/Source/USDUtilities/Public/USDSkeletalDataConversion.h b/Engine/Plugins/Runtime/USDCore/Source/USDUtilities/Public/USDSkeletalDataConversion.h
--- a/Engine/Plugins/Runtime/USDCore/Source/USDUtilities/Public/USDSkeletalDataConversion.h
+++ b/Engine/Plugins/Runtime/USDCore/Source/USDUtilities/Public/USDSkeletalDataConversion.h
@@ -190,7 +190,8 @@ namespace UsdUtils
 	USDUTILITIES_API bool ApplyBlendShape(
 		FMeshDescription& InOutMeshDescription,
 		const pxr::UsdPrim& InBlendShapePrim,
-		float Weight = 1.0f,
+		const FTransform& AdditionalTransform,
+		float Weight = 1.0f,		
 		const FString InInbetweenName = {}
 	);
 #endif	  // USE_USD_SDK
@@ -256,7 +257,8 @@ namespace UsdToUnreal
 		const FUsdStageInfo& StageInfo,
 		uint32 PointIndexOffset,
 		TSet<FString>& UsedMorphTargetNames,
-		UsdUtils::FBlendShapeMap& OutBlendShapes
+		UsdUtils::FBlendShapeMap& OutBlendShapes,
+		const FUsdMeshConversionOptions& Options = FUsdMeshConversionOptions::DefaultOptions
 	);
 	USDUTILITIES_API bool ConvertBlendShape(
 		const pxr::UsdSkelBlendShape& UsdBlendShape,
@@ -264,7 +266,8 @@ namespace UsdToUnreal
 		int32 LODIndex,
 		uint32 PointIndexOffset,
 		TSet<FString>& UsedMorphTargetNames,
-		UsdUtils::FBlendShapeMap& OutBlendShapes
+		UsdUtils::FBlendShapeMap& OutBlendShapes,
+		const FUsdMeshConversionOptions& Options
 	);
 
 	/**
