USD: Refactor error and logging utilities.

diff --git a/Engine/Plugins/Importers/USDImporter/Source/GeometryCacheUSD/Private/GeometryCacheTrackUSD.cpp b/Engine/Plugins/Importers/USDImporter/Source/GeometryCacheUSD/Private/GeometryCacheTrackUSD.cpp
--- a/Engine/Plugins/Importers/USDImporter/Source/GeometryCacheUSD/Private/GeometryCacheTrackUSD.cpp
+++ b/Engine/Plugins/Importers/USDImporter/Source/GeometryCacheUSD/Private/GeometryCacheTrackUSD.cpp
@@ -2,7 +2,7 @@
 
 #include "GeometryCacheTrackUSD.h"
 
-#include "USDLog.h"
+#include "USDErrorUtils.h"
 #include "UsdWrappers/SdfLayer.h"
 #include "UsdWrappers/UsdStage.h"
 
@@ -222,9 +222,7 @@ bool UGeometryCacheTrackUsd::LoadUsdStage()
 	}
 	else if (!StageRootLayerPath.IsEmpty())
 	{
-		UE_LOG(
-			LogUsd,
-			Warning,
+		USD_LOG_WARNING(
 			TEXT("UGeometryCacheTrackUsd is reopening the stage '%s' to stream in frames for the geometry cache generated for prim '%s'"),
 			*StageRootLayerPath,
 			*PrimPath
@@ -245,7 +243,7 @@ bool UGeometryCacheTrackUsd::LoadUsdStage()
 		return true;
 	}
 
-	UE_LOG(LogUsd, Warning, TEXT("UGeometryCacheTrackUsd track failed to access USD stage to stream requested frames"));
+	USD_LOG_WARNING(TEXT("UGeometryCacheTrackUsd track failed to access USD stage to stream requested frames"));
 	return false;
 }
 
@@ -273,7 +271,7 @@ void UGeometryCacheTrackUsd::Initialize(
 
 	if (FramesPerSecond == 0)
 	{
-		ensureMsgf(false, TEXT("Invalid USD GeometryCache FPS detected. Falling back to 1 FPS"));
+		USD_LOG_WARNING(TEXT("Invalid USD GeometryCache FPS detected. Falling back to 1 FPS"));
 		FramesPerSecond = 1;
 	}
 
diff --git a/Engine/Plugins/Importers/USDImporter/Source/USDExporter/Private/AnimSequenceExporterUSD.cpp b/Engine/Plugins/Importers/USDImporter/Source/USDExporter/Private/AnimSequenceExporterUSD.cpp
--- a/Engine/Plugins/Importers/USDImporter/Source/USDExporter/Private/AnimSequenceExporterUSD.cpp
+++ b/Engine/Plugins/Importers/USDImporter/Source/USDExporter/Private/AnimSequenceExporterUSD.cpp
@@ -8,7 +8,6 @@
 #include "USDExporterModule.h"
 #include "USDExportUtils.h"
 #include "USDLayerUtils.h"
-#include "USDLog.h"
 #include "USDObjectUtils.h"
 #include "USDOptionsWindow.h"
 #include "USDPrimConversion.h"
@@ -28,6 +27,8 @@
 #include "Misc/EngineVersion.h"
 #include "UObject/GCObjectScopeGuard.h"
 
+#define LOCTEXT_NAMESPACE "AnimSequenceExporterUSD"
+
 namespace UE::AnimSequenceExporterUSD::Private
 {
 	void SendAnalytics(
@@ -190,13 +191,10 @@ bool UAnimSequenceExporterUSD::ExportBinary(
 		}
 		else
 		{
-			FUsdLogManager::LogMessage(
-				EMessageSeverity::Warning,
-				FText::Format(
-					NSLOCTEXT("AnimSequenceExporterUSD", "InvalidSkelMesh", "Couldn't find the skeletal mesh to export for anim sequence {0}."),
-					FText::FromName(AnimSequence->GetFName())
-				)
-			);
+			USD_LOG_USERWARNING(FText::Format(
+				NSLOCTEXT("AnimSequenceExporterUSD", "InvalidSkelMesh", "Couldn't find the skeletal mesh to export for anim sequence {0}."),
+				FText::FromName(AnimSequence->GetFName())
+			));
 		}
 	}
 
@@ -235,13 +233,11 @@ bool UAnimSequenceExporterUSD::ExportBinary(
 	{
 		if (!ExportTask->bReplaceIdentical)
 		{
-			UE_LOG(
-				LogUsd,
-				Log,
-				TEXT("Skipping export of asset '%s' as the target file '%s' already exists."),
-				*Object->GetPathName(),
-				*AnimSequenceFile
-			);
+			USD_LOG_USERINFO(FText::Format(
+				LOCTEXT("FileAlreadyExists", "Skipping export of asset '{0}' as the target file '{1}' already exists."),
+				FText::FromString(Object->GetPathName()),
+				FText::FromString(UExporter::CurrentFilename)
+			));
 			return false;
 		}
 		// If we don't want to re-export this asset we need to check if its the same version
@@ -263,13 +259,14 @@ bool UAnimSequenceExporterUSD::ExportBinary(
 
 					if (bVersionMatches && bAssetTypeMatches)
 					{
-						UE_LOG(
-							LogUsd,
-							Log,
-							TEXT("Skipping export of asset '%s' as the target file '%s' already contains up-to-date exported data."),
-							*Object->GetPathName(),
-							*AnimSequenceFile
-						);
+						USD_LOG_USERINFO(FText::Format(
+							LOCTEXT(
+								"FileUpToDate",
+								"Skipping export of asset '{0}' as the target file '{1}' already contains up-to-date exported data."
+							),
+							FText::FromString(Object->GetPathName()),
+							FText::FromString(UExporter::CurrentFilename)
+						));
 						return true;
 					}
 				}
@@ -422,3 +419,5 @@ bool UAnimSequenceExporterUSD::ExportBinary(
 	return false;
 #endif	  // #if USE_USD_SDK
 }
+
+#undef LOCTEXT_NAMESPACE
diff --git a/Engine/Plugins/Importers/USDImporter/Source/USDExporter/Private/GeometryCacheExporterUSD.cpp b/Engine/Plugins/Importers/USDImporter/Source/USDExporter/Private/GeometryCacheExporterUSD.cpp
--- a/Engine/Plugins/Importers/USDImporter/Source/USDExporter/Private/GeometryCacheExporterUSD.cpp
+++ b/Engine/Plugins/Importers/USDImporter/Source/USDExporter/Private/GeometryCacheExporterUSD.cpp
@@ -6,10 +6,10 @@
 #include "MaterialExporterUSD.h"
 #include "UnrealUSDWrapper.h"
 #include "USDConversionUtils.h"
+#include "USDErrorUtils.h"
 #include "USDExporterModule.h"
 #include "USDExportUtils.h"
 #include "USDGeomMeshConversion.h"
-#include "USDLog.h"
 #include "USDObjectUtils.h"
 #include "USDOptionsWindow.h"
 #include "USDPrimConversion.h"
@@ -25,6 +25,8 @@
 #include "GeometryCache.h"
 #include "Misc/EngineVersion.h"
 
+#define LOCTEXT_NAMESPACE "GeometryCacheExporterUSD"
+
 namespace UE::GeometryCacheExporterUSD::Private
 {
 	void SendAnalytics(
@@ -168,13 +170,11 @@ bool UGeometryCacheExporterUSD::ExportBinary(
 	{
 		if (!ExportTask->bReplaceIdentical)
 		{
-			UE_LOG(
-				LogUsd,
-				Log,
-				TEXT("Skipping export of asset '%s' as the target file '%s' already exists."),
-				*Object->GetPathName(),
-				*UExporter::CurrentFilename
-			);
+			USD_LOG_USERINFO(FText::Format(
+				LOCTEXT("FileAlreadyExists", "Skipping export of asset '{0}' as the target file '{1}' already exists."),
+				FText::FromString(Object->GetPathName()),
+				FText::FromString(UExporter::CurrentFilename)
+			));
 			return false;
 		}
 		// If we don't want to re-export this asset we need to check if its the same version
@@ -198,13 +198,14 @@ bool UGeometryCacheExporterUSD::ExportBinary(
 
 					if (bVersionMatches && bAssetTypeMatches)
 					{
-						UE_LOG(
-							LogUsd,
-							Log,
-							TEXT("Skipping export of asset '%s' as the target file '%s' already contains up-to-date exported data."),
-							*GeometryCache->GetPathName(),
-							*UExporter::CurrentFilename
-						);
+						USD_LOG_USERINFO(FText::Format(
+							LOCTEXT(
+								"FileUpToDate",
+								"Skipping export of asset '{0}' as the target file '{1}' already contains up-to-date exported data."
+							),
+							FText::FromString(GeometryCache->GetPathName()),
+							FText::FromString(UExporter::CurrentFilename)
+						));
 
 						bSkipMeshExport = true;
 					}
@@ -372,3 +373,5 @@ bool UGeometryCacheExporterUSD::ExportBinary(
 	return false;
 #endif	  // #if USE_USD_SDK
 }
+
+#undef LOCTEXT_NAMESPACE
diff --git a/Engine/Plugins/Importers/USDImporter/Source/USDExporter/Private/LevelSequenceExporterUSD.cpp b/Engine/Plugins/Importers/USDImporter/Source/USDExporter/Private/LevelSequenceExporterUSD.cpp
--- a/Engine/Plugins/Importers/USDImporter/Source/USDExporter/Private/LevelSequenceExporterUSD.cpp
+++ b/Engine/Plugins/Importers/USDImporter/Source/USDExporter/Private/LevelSequenceExporterUSD.cpp
@@ -12,7 +12,6 @@
 #include "USDExportUtils.h"
 #include "USDGeomMeshConversion.h"
 #include "USDLayerUtils.h"
-#include "USDLog.h"
 #include "USDObjectUtils.h"
 #include "USDOptionsWindow.h"
 #include "USDPrimConversion.h"
@@ -162,18 +161,6 @@ namespace UE::LevelSequenceExporterUSD::Private
 			if (int32* ExistingIndex = SpawnableIndices.Find(InstanceKey))
 			{
 				Object = ExistingInstancesForGuid[*ExistingIndex];
-
-				UE_LOG(
-					LogUsd,
-					VeryVerbose,
-					TEXT("Already spawned '%s' (%0x) for RootSequence '%s', TemplateID '%u', Guid '%s' (index %d)"),
-					*Object->GetPathName(),
-					Object,
-					*RootSequence->GetPathName(),
-					TemplateID.GetInternalValue(),
-					*Guid.ToString(),
-					*ExistingIndex
-				);
 			}
 			else if (!bExportSeparatePrimsPerSpawnableInstance && ExistingInstancesForGuid.Num() > 0)
 			{
@@ -214,18 +201,6 @@ namespace UE::LevelSequenceExporterUSD::Private
 				{
 					Object = ExistingInstancesForGuid[IndexToReuse];
 					SpawnableIndices.Add(InstanceKey, IndexToReuse);
-
-					UE_LOG(
-						LogUsd,
-						VeryVerbose,
-						TEXT("Reusing '%s' (%0x) for RootSequence '%s', TemplateID '%u', Guid '%s' (index %d)"),
-						*Object->GetPathName(),
-						Object,
-						*RootSequence->GetPathName(),
-						TemplateID.GetInternalValue(),
-						*Guid.ToString(),
-						IndexToReuse
-					);
 				}
 			}
 
@@ -242,16 +217,6 @@ namespace UE::LevelSequenceExporterUSD::Private
 			{
 				// SpawnObject will add an entry into the Register for us
 				Object = FMovieSceneSpawnRegister::SpawnObject(Guid, MovieScene, TemplateID, SharedPlaybackState, BindingIndex);
-				UE_LOG(
-					LogUsd,
-					VeryVerbose,
-					TEXT("Spawning '%s' (%0x) for RootSequence '%s', TemplateID '%u', Guid '%s'"),
-					*Object->GetPathName(),
-					Object,
-					*RootSequence->GetPathName(),
-					TemplateID.GetInternalValue(),
-					*Guid.ToString()
-				);
 
 				ExistingInstancesForGuid.Add(Object);
 				SpawnableIndices.Add(InstanceKey, ExistingInstancesForGuid.Num() - 1);
@@ -791,13 +756,14 @@ namespace UE::LevelSequenceExporterUSD::Private
 			// section. If we tried exporting another section here we'd need a fully separate prim for it, which opens
 			// a can of worms as we so far only had one prim per binding. Furthermore we'd need to pay attention to this
 			// split when opening the stage as well, otherwise we'd roundtrip the audio track incorrectly
-			UE_LOG(
-				LogUsd,
-				Warning,
-				TEXT("The audio track '%s' has %d sections, but only the first audio section of an audio track can be written out to USD for now"),
-				*AudioTrack.GetPathName(),
+			USD_LOG_USERWARNING(FText::Format(
+				LOCTEXT(
+					"MultipleAudioSections",
+					"The audio track '{0}' has {1} sections, but only the first audio section of an audio track can be written out to USD for now"
+				),
+				FText::FromString(AudioTrack.GetPathName()),
 				Sections.Num()
-			);
+			));
 		}
 
 		if (Sections.Num() > 0)
@@ -813,15 +779,14 @@ namespace UE::LevelSequenceExporterUSD::Private
 		{
 			if (*SourceTracksForPrim > 1)
 			{
-				UE_LOG(
-					LogUsd,
-					Warning,
-					TEXT(
-						"Exporting multiple audio tracks (like '%s') to the same prim ('%s') is currently not supported and may lead to incorrect output"
+				USD_LOG_USERWARNING(FText::Format(
+					LOCTEXT(
+						"MultipleAudioTracks",
+						"Exporting multiple audio tracks (like '{0}') to the same prim ('{1}') is currently not supported and may lead to incorrect output"
 					),
-					*AudioTrack.GetPathName(),
-					*PrimPath
-				);
+					FText::FromString(AudioTrack.GetPathName()),
+					FText::FromString(PrimPath)
+				));
 			}
 		}
 	}
@@ -1087,18 +1052,15 @@ namespace UE::LevelSequenceExporterUSD::Private
 								}
 								else
 								{
-									FUsdLogManager::LogMessage(
-										EMessageSeverity::Warning,
-										FText::Format(
-											LOCTEXT(
-												"NonIdealComposition",
-												"Exported animation for prim '{0}' may not compose correctly with the prims from referenced layer '{1}' on the exported stage for the LevelSequence '{2}'. For best results, make sure the referenced layer is saved to disk (i.e. not anonymous), has a defaultPrim setup, and that the animation tracks are only bound to prims that are descendents of the defaultPrim."
-											),
-											FText::FromString(PrimPathInSourceStage),
-											FText::FromString(LoadedStage.GetRootLayer().GetIdentifier()),
-											FText::FromString(Context.RootSequence.GetPathName())
-										)
-									);
+									USD_LOG_USERWARNING(FText::Format(
+										LOCTEXT(
+											"NonIdealComposition",
+											"Exported animation for prim '{0}' may not compose correctly with the prims from referenced layer '{1}' on the exported stage for the LevelSequence '{2}'. For best results, make sure the referenced layer is saved to disk (i.e. not anonymous), has a defaultPrim setup, and that the animation tracks are only bound to prims that are descendents of the defaultPrim."
+										),
+										FText::FromString(PrimPathInSourceStage),
+										FText::FromString(LoadedStage.GetRootLayer().GetIdentifier()),
+										FText::FromString(Context.RootSequence.GetPathName())
+									));
 								}
 							}
 						}
@@ -1169,12 +1131,7 @@ namespace UE::LevelSequenceExporterUSD::Private
 						}
 						else
 						{
-							UE_LOG(
-								LogUsd,
-								Warning,
-								TEXT("Failed to generate Skeleton or SkelAnimation prim when baking out SkelRoot '%s'"),
-								*PrimPath
-							);
+							USD_LOG_WARNING(TEXT("Failed to generate Skeleton or SkelAnimation prim when baking out SkelRoot '%s'"), *PrimPath);
 						}
 					}
 					else
@@ -1618,13 +1575,11 @@ namespace UE::LevelSequenceExporterUSD::Private
 		{
 			if (!Context.bReplaceIdentical)
 			{
-				UE_LOG(
-					LogUsd,
-					Log,
-					TEXT("Skipping export of asset '%s' as the target file '%s' already exists."),
-					*MovieSceneSequence.GetPathName(),
-					*UExporter::CurrentFilename
-				);
+				USD_LOG_USERINFO(FText::Format(
+					LOCTEXT("FileAlreadyExists", "Skipping export of asset '{0}' as the target file '{1}' already exists."),
+					FText::FromString(MovieSceneSequence.GetPathName()),
+					FText::FromString(UExporter::CurrentFilename)
+				));
 				return;
 			}
 			// If we don't want to re-export this level sequence we need to check if its the same version
@@ -1647,13 +1602,14 @@ namespace UE::LevelSequenceExporterUSD::Private
 
 						if (bVersionMatches && bAssetTypeMatches)
 						{
-							UE_LOG(
-								LogUsd,
-								Log,
-								TEXT("Skipping export of asset '%s' as the target file '%s' already contains up-to-date exported data."),
-								*MovieSceneSequence.GetPathName(),
-								*UniqueFilePath
-							);
+							USD_LOG_USERINFO(FText::Format(
+								LOCTEXT(
+									"FileUpToDate",
+									"Skipping export of asset '{0}' as the target file '{1}' already contains up-to-date exported data."
+								),
+								FText::FromString(MovieSceneSequence.GetPathName()),
+								FText::FromString(UExporter::CurrentFilename)
+							));
 							return;
 						}
 					}
diff --git a/Engine/Plugins/Importers/USDImporter/Source/USDExporter/Private/MaterialExporterUSD.cpp b/Engine/Plugins/Importers/USDImporter/Source/USDExporter/Private/MaterialExporterUSD.cpp
--- a/Engine/Plugins/Importers/USDImporter/Source/USDExporter/Private/MaterialExporterUSD.cpp
+++ b/Engine/Plugins/Importers/USDImporter/Source/USDExporter/Private/MaterialExporterUSD.cpp
@@ -6,10 +6,10 @@
 #include "UnrealUSDWrapper.h"
 #include "USDClassesModule.h"
 #include "USDConversionUtils.h"
+#include "USDErrorUtils.h"
 #include "USDExporterModule.h"
 #include "USDExportUtils.h"
 #include "USDGeomMeshConversion.h"
-#include "USDLog.h"
 #include "USDMetadataExportOptions.h"
 #include "USDObjectUtils.h"
 #include "USDOptionsWindow.h"
@@ -37,6 +37,8 @@
 #include "UObject/UObjectGlobals.h"
 #include "VT/RuntimeVirtualTexture.h"
 
+#define LOCTEXT_NAMESPACE "MaterialExporterUSD"
+
 namespace UE::MaterialExporterUSD::Private
 {
 	void SendAnalytics(
@@ -289,13 +291,11 @@ bool UMaterialExporterUsd::ExportMaterial(
 	{
 		if (!bReplaceIdentical)
 		{
-			UE_LOG(
-				LogUsd,
-				Log,
-				TEXT("Skipping export of asset '%s' as the target file '%s' already exists."),
-				*Material.GetPathName(),
-				*UExporter::CurrentFilename
-			);
+			USD_LOG_USERINFO(FText::Format(
+				LOCTEXT("FileAlreadyExists", "Skipping export of asset '{0}' as the target file '{1}' already exists."),
+				FText::FromString(Material.GetPathName()),
+				FText::FromString(UExporter::CurrentFilename)
+			));
 			return false;
 		}
 		// If we don't want to re-export this asset we need to check if its the same version
@@ -317,13 +317,15 @@ bool UMaterialExporterUsd::ExportMaterial(
 
 					if (bVersionMatches && bAssetTypeMatches)
 					{
-						UE_LOG(
-							LogUsd,
-							Log,
-							TEXT("Skipping export of asset '%s' as the target file '%s' already contains up-to-date exported data."),
-							*Material.GetPathName(),
-							*FilePath.FilePath
-						);
+						USD_LOG_USERINFO(FText::Format(
+							LOCTEXT(
+								"FileUpToDate",
+								"Skipping export of asset '{0}' as the target file '{1}' already contains up-to-date exported data."
+							),
+							FText::FromString(Material.GetPathName()),
+							FText::FromString(UExporter::CurrentFilename)
+						));
+
 						return true;
 					}
 				}
@@ -551,3 +553,5 @@ bool UMaterialExporterUsd::ExportMaterialsForStage(
 		bIsAutomated
 	);
 }
+
+#undef LOCTEXT_NAMESPACE
diff --git a/Engine/Plugins/Importers/USDImporter/Source/USDExporter/Private/SkeletalMeshExporterUSD.cpp b/Engine/Plugins/Importers/USDImporter/Source/USDExporter/Private/SkeletalMeshExporterUSD.cpp
--- a/Engine/Plugins/Importers/USDImporter/Source/USDExporter/Private/SkeletalMeshExporterUSD.cpp
+++ b/Engine/Plugins/Importers/USDImporter/Source/USDExporter/Private/SkeletalMeshExporterUSD.cpp
@@ -6,9 +6,9 @@
 #include "USDAssetUserData.h"
 #include "USDClassesModule.h"
 #include "USDConversionUtils.h"
+#include "USDErrorUtils.h"
 #include "USDExporterModule.h"
 #include "USDExportUtils.h"
-#include "USDLog.h"
 #include "USDObjectUtils.h"
 #include "USDOptionsWindow.h"
 #include "USDPrimConversion.h"
@@ -29,6 +29,8 @@
 #include "Misc/EngineVersion.h"
 #include "Rendering/SkeletalMeshRenderData.h"
 
+#define LOCTEXT_NAMESPACE "SkeletalMeshExporterUSD"
+
 namespace UE::SkeletalMeshExporterUSD::Private
 {
 	void SendAnalytics(
@@ -186,13 +188,11 @@ bool USkeletalMeshExporterUsd::ExportBinary(
 	{
 		if (!ExportTask->bReplaceIdentical)
 		{
-			UE_LOG(
-				LogUsd,
-				Log,
-				TEXT("Skipping export of asset '%s' as the target file '%s' already exists."),
-				*Object->GetPathName(),
-				*UExporter::CurrentFilename
-			);
+			USD_LOG_USERINFO(FText::Format(
+				LOCTEXT("FileAlreadyExists", "Skipping export of asset '{0}' as the target file '{1}' already exists."),
+				FText::FromString(Object->GetPathName()),
+				FText::FromString(UExporter::CurrentFilename)
+			));
 			return false;
 		}
 		// If we don't want to re-export this asset we need to check if its the same version
@@ -216,13 +216,14 @@ bool USkeletalMeshExporterUsd::ExportBinary(
 
 					if (bVersionMatches && bAssetTypeMatches)
 					{
-						UE_LOG(
-							LogUsd,
-							Log,
-							TEXT("Skipping export of asset '%s' as the target file '%s' already contains up-to-date exported data."),
-							*SkeletalMesh->GetPathName(),
-							*UExporter::CurrentFilename
-						);
+						USD_LOG_USERINFO(FText::Format(
+							LOCTEXT(
+								"FileUpToDate",
+								"Skipping export of asset '{0}' as the target file '{1}' already contains up-to-date exported data."
+							),
+							FText::FromString(SkeletalMesh->GetPathName()),
+							FText::FromString(UExporter::CurrentFilename)
+						));
 
 						bSkipMeshExport = true;
 					}
@@ -426,3 +427,5 @@ bool USkeletalMeshExporterUsd::ExportBinary(
 	return false;
 #endif	  // #if USE_USD_SDK
 }
+
+#undef LOCTEXT_NAMESPACE
diff --git a/Engine/Plugins/Importers/USDImporter/Source/USDExporter/Private/StaticMeshExporterUSD.cpp b/Engine/Plugins/Importers/USDImporter/Source/USDExporter/Private/StaticMeshExporterUSD.cpp
--- a/Engine/Plugins/Importers/USDImporter/Source/USDExporter/Private/StaticMeshExporterUSD.cpp
+++ b/Engine/Plugins/Importers/USDImporter/Source/USDExporter/Private/StaticMeshExporterUSD.cpp
@@ -6,10 +6,10 @@
 #include "StaticMeshExporterUSDOptions.h"
 #include "USDClassesModule.h"
 #include "USDConversionUtils.h"
+#include "USDErrorUtils.h"
 #include "USDExporterModule.h"
 #include "USDExportUtils.h"
 #include "USDGeomMeshConversion.h"
-#include "USDLog.h"
 #include "USDObjectUtils.h"
 #include "USDOptionsWindow.h"
 #include "USDPrimConversion.h"
@@ -27,6 +27,8 @@
 #include "Misc/EngineVersion.h"
 #include "StaticMeshResources.h"
 
+#define LOCTEXT_NAMESPACE "StaticMeshExporterUSD"
+
 namespace UE::StaticMeshExporterUSD::Private
 {
 	void SendAnalytics(
@@ -191,13 +193,11 @@ bool UStaticMeshExporterUsd::ExportBinary(UObject* Object, const TCHAR* Type, FA
 	{
 		if (!ExportTask->bReplaceIdentical)
 		{
-			UE_LOG(
-				LogUsd,
-				Log,
-				TEXT("Skipping export of asset '%s' as the target file '%s' already exists."),
-				*Object->GetPathName(),
-				*UExporter::CurrentFilename
-			);
+			USD_LOG_USERINFO(FText::Format(
+				LOCTEXT("FileAlreadyExists", "Skipping export of asset '{0}' as the target file '{1}' already exists."),
+				FText::FromString(Object->GetPathName()),
+				FText::FromString(UExporter::CurrentFilename)
+			));
 			return false;
 		}
 		// If we don't want to re-export this asset we need to check if its the same version
@@ -221,13 +221,14 @@ bool UStaticMeshExporterUsd::ExportBinary(UObject* Object, const TCHAR* Type, FA
 
 					if (bVersionMatches && bAssetTypeMatches)
 					{
-						UE_LOG(
-							LogUsd,
-							Log,
-							TEXT("Skipping export of asset '%s' as the target file '%s' already contains up-to-date exported data."),
-							*StaticMesh->GetPathName(),
-							*UExporter::CurrentFilename
-						);
+						USD_LOG_USERINFO(FText::Format(
+							LOCTEXT(
+								"FileUpToDate",
+								"Skipping export of asset '{0}' as the target file '{1}' already contains up-to-date exported data."
+							),
+							FText::FromString(StaticMesh->GetPathName()),
+							FText::FromString(UExporter::CurrentFilename)
+						));
 
 						bSkipMeshExport = true;
 					}
@@ -403,3 +404,5 @@ bool UStaticMeshExporterUsd::ExportBinary(UObject* Object, const TCHAR* Type, FA
 	return false;
 #endif	  // #if USE_USD_SDK
 }
+
+#undef LOCTEXT_NAMESPACE
diff --git a/Engine/Plugins/Importers/USDImporter/Source/USDExporter/Private/USDConversionBlueprintContext.cpp b/Engine/Plugins/Importers/USDImporter/Source/USDExporter/Private/USDConversionBlueprintContext.cpp
--- a/Engine/Plugins/Importers/USDImporter/Source/USDExporter/Private/USDConversionBlueprintContext.cpp
+++ b/Engine/Plugins/Importers/USDImporter/Source/USDExporter/Private/USDConversionBlueprintContext.cpp
@@ -4,9 +4,9 @@
 
 #include "PrimitiveSceneProxy.h"
 #include "USDConversionUtils.h"
+#include "USDErrorUtils.h"
 #include "USDGeomMeshConversion.h"
 #include "USDLightConversion.h"
-#include "USDLog.h"
 #include "USDMetadata.h"
 #include "USDPrimConversion.h"
 #include "USDShadeConversion.h"
@@ -25,13 +25,15 @@
 #include "RHI.h"
 #include "StaticMeshAttributes.h"
 
+#define LOCTEXT_NAMESPACE "USDConversionBlueprintContext"
+
 namespace UE::USDExporter::Private
 {
 	UE::FUsdPrim GetPrim(const UE::FUsdStage& Stage, const FString& PrimPath)
 	{
 		if (!Stage)
 		{
-			UE_LOG(LogUsd, Error, TEXT("Export context has no stage set! Call SetStage with a root layer filepath first."));
+			USD_LOG_USERERROR(LOCTEXT("MissingStage", "Export context has no stage set! Call SetStage with a root layer filepath first."));
 			return {};
 		}
 
@@ -92,7 +94,7 @@ FFilePath UUsdConversionBlueprintContext::GetStageRootLayer()
 		return FFilePath{Stage.GetRootLayer().GetRealPath()};
 	}
 
-	UE_LOG(LogUsd, Error, TEXT("There is no stage currently open!"));
+	USD_LOG_USERERROR(LOCTEXT("NoStage", "There is no stage currently open!"));
 	return {};
 }
 
@@ -106,12 +108,15 @@ void UUsdConversionBlueprintContext::SetEditTarget(FFilePath EditTargetLayerPath
 		}
 		else
 		{
-			UE_LOG(LogUsd, Error, TEXT("Failed to find or open USD layer with filepath '%s'!"), *EditTargetLayerPath.FilePath);
+			USD_LOG_USERERROR(FText::Format(
+				LOCTEXT("FailedToOpenPath", "Failed to find or open USD layer with filepath '{0}'!"),
+				FText::FromString(EditTargetLayerPath.FilePath)
+			));
 		}
 	}
 	else
 	{
-		UE_LOG(LogUsd, Error, TEXT("There is no stage currently open!"));
+		USD_LOG_USERERROR(LOCTEXT("NoStage", "There is no stage currently open!"));
 	}
 }
 
@@ -122,7 +127,7 @@ FFilePath UUsdConversionBlueprintContext::GetEditTarget()
 		return FFilePath{Stage.GetEditTarget().GetRealPath()};
 	}
 
-	UE_LOG(LogUsd, Error, TEXT("There is no stage currently open!"));
+	USD_LOG_USERERROR(LOCTEXT("NoStage", "There is no stage currently open!"));
 	return {};
 }
 
@@ -413,7 +418,11 @@ bool UUsdConversionBlueprintContext::ConvertLandscapeProxyActorMesh(
 		ExportParams.ExportLOD = LODIndex + LowestLOD;
 		if (!Actor->ExportToRawMesh(ExportParams, MeshDescription))
 		{
-			UE_LOG(LogUsd, Error, TEXT("Failed to convert LOD %d of landscape actor '%s''s mesh data"), LODIndex, *Actor->GetName());
+			USD_LOG_USERERROR(FText::Format(
+				LOCTEXT("FailedLODLandscape", "Failed to convert LOD {0} of landscape actor '{1}''s mesh data"),
+				LODIndex,
+				FText::FromString(Actor->GetName())
+			));
 			return false;
 		}
 	}
@@ -685,7 +694,9 @@ FUsdCombinedPrimMetadata UUsdConversionBlueprintContext::GetPrimMetadata(
 			bool bSuccess = UsdToUnreal::ConvertMetadata(Prim, Result, BlockedPrefixFilter, bInvertFilter, bCollectFromEntireSubtrees);
 			if (!bSuccess)
 			{
-				UE_LOG(LogUsd, Warning, TEXT("Failed to get metadata from prim '%s'"), *PrimPath);
+				USD_LOG_USERWARNING(
+					FText::Format(LOCTEXT("FailedToGetMetadata", "Failed to get metadata from prim '{0}'"), FText::FromString(PrimPath))
+				);
 				Result = {};
 			}
 		}
@@ -694,3 +705,5 @@ FUsdCombinedPrimMetadata UUsdConversionBlueprintContext::GetPrimMetadata(
 
 	return Result;
 }
+
+#undef LOCTEXT_NAMESPACE
diff --git a/Engine/Plugins/Importers/USDImporter/Source/USDExporter/Private/USDConversionBlueprintLibrary.cpp b/Engine/Plugins/Importers/USDImporter/Source/USDExporter/Private/USDConversionBlueprintLibrary.cpp
--- a/Engine/Plugins/Importers/USDImporter/Source/USDExporter/Private/USDConversionBlueprintLibrary.cpp
+++ b/Engine/Plugins/Importers/USDImporter/Source/USDExporter/Private/USDConversionBlueprintLibrary.cpp
@@ -7,10 +7,10 @@
 #include "USDAssetUserData.h"
 #include "USDClassesModule.h"
 #include "USDConversionUtils.h"
+#include "USDErrorUtils.h"
 #include "USDExporterModule.h"
 #include "USDExportUtils.h"
 #include "USDLayerUtils.h"
-#include "USDLog.h"
 #include "USDObjectUtils.h"
 #include "USDValueConversion.h"
 
@@ -34,6 +34,8 @@
 #include "LevelEditorSequencerIntegration.h"
 #include "UObject/ObjectMacros.h"
 
+#define LOCTEXT_NAMESPACE "USDConversionBlueprintLibrary"
+
 namespace UE::UsdConversionBlueprintLibrary::Private
 {
 	void WaitForAllAsyncAndSteamingTasks(UWorld* World)
@@ -404,7 +406,11 @@ FString UUsdConversionBlueprintLibrary::MakePathRelativeToLayer(const FString& A
 	}
 	else
 	{
-		UE_LOG(LogUsd, Error, TEXT("Failed to find a layer with path '%s' to make the path '%s' relative to"), *AnchorLayerPath, *PathToMakeRelative);
+		USD_LOG_USERERROR(FText::Format(
+			LOCTEXT("FailedToFindAnchor", "Failed to find a layer with path '{0}' to make the path '{1}' relative to"),
+			FText::FromString(AnchorLayerPath),
+			FText::FromString(PathToMakeRelative)
+		));
 		return PathToMakeRelative;
 	}
 #else
@@ -426,7 +432,11 @@ void UUsdConversionBlueprintLibrary::InsertSubLayer(const FString& ParentLayerPa
 	}
 	else
 	{
-		UE_LOG(LogUsd, Error, TEXT("Failed to find a parent layer '%s' when trying to insert sublayer '%s'"), *ParentLayerPath, *SubLayerPath);
+		USD_LOG_USERERROR(FText::Format(
+			LOCTEXT("FailedToFindParent", "Failed to find a parent layer '{0}' when trying to insert sublayer '{1}'"),
+			FText::FromString(ParentLayerPath),
+			FText::FromString(SubLayerPath)
+		));
 	}
 #endif	  // USE_USD_SDK
 }
@@ -896,13 +906,14 @@ bool UUsdConversionBlueprintLibrary::SetMetadataField(
 
 	if (StageIdentifier.IsEmpty() && AssetUserData->StageIdentifierToMetadata.Num() != 1)
 	{
-		UE_LOG(
-			LogUsd,
-			Warning,
-			TEXT("Failed to set metadata field '%s' on AssetUserData '%s': Please provide a valid value for the StageIdentifier parameter"),
-			*AssetUserData->GetPathName(),
-			*Key
-		);
+		USD_LOG_USERWARNING(FText::Format(
+			LOCTEXT(
+				"InvalidStageIdentifier",
+				"Failed to set metadata field '{0}' on AssetUserData '{1}': Please provide a valid value for the StageIdentifier parameter"
+			),
+			FText::FromString(AssetUserData->GetPathName()),
+			FText::FromString(Key)
+		));
 		return false;
 	}
 
@@ -912,13 +923,14 @@ bool UUsdConversionBlueprintLibrary::SetMetadataField(
 
 	if (PrimPath.IsEmpty() && CombinedPrimMetadata.PrimPathToMetadata.Num() != 1)
 	{
-		UE_LOG(
-			LogUsd,
-			Warning,
-			TEXT("Failed to set metadata field '%s' on AssetUserData '%s': Please provide a valid value for the PrimPath parameter"),
-			*AssetUserData->GetPathName(),
-			*Key
-		);
+		USD_LOG_USERWARNING(FText::Format(
+			LOCTEXT(
+				"InvalidPrimPath",
+				"Failed to set metadata field '{0}' on AssetUserData '{1}': Please provide a valid value for the PrimPath parameter"
+			),
+			FText::FromString(AssetUserData->GetPathName()),
+			FText::FromString(Key)
+		));
 		return false;
 	}
 
@@ -960,13 +972,14 @@ bool UUsdConversionBlueprintLibrary::ClearMetadataField(
 
 	if (StageIdentifier.IsEmpty() && AssetUserData->StageIdentifierToMetadata.Num() != 1)
 	{
-		UE_LOG(
-			LogUsd,
-			Warning,
-			TEXT("Failed to clear metadata field '%s' on AssetUserData '%s': Please provide a valid value for the StageIdentifier parameter"),
-			*AssetUserData->GetPathName(),
-			*Key
-		);
+		USD_LOG_USERWARNING(FText::Format(
+			LOCTEXT(
+				"InvalidStageIdentifierClear",
+				"Failed to clear metadata field '{0}' on AssetUserData '{1}': Please provide a valid value for the StageIdentifier parameter"
+			),
+			FText::FromString(AssetUserData->GetPathName()),
+			FText::FromString(Key)
+		));
 		return false;
 	}
 
@@ -976,13 +989,14 @@ bool UUsdConversionBlueprintLibrary::ClearMetadataField(
 
 	if (PrimPath.IsEmpty() && CombinedPrimMetadata.PrimPathToMetadata.Num() != 1)
 	{
-		UE_LOG(
-			LogUsd,
-			Warning,
-			TEXT("Failed to clear metadata field '%s' on AssetUserData '%s': Please provide a valid value for the PrimPath parameter"),
-			*AssetUserData->GetPathName(),
-			*Key
-		);
+		USD_LOG_USERWARNING(FText::Format(
+			LOCTEXT(
+				"InvalidPrimPathClear",
+				"Failed to clear metadata field '{0}' on AssetUserData '{1}': Please provide a valid value for the PrimPath parameter"
+			),
+			FText::FromString(AssetUserData->GetPathName()),
+			FText::FromString(Key)
+		));
 		return false;
 	}
 
@@ -1019,13 +1033,14 @@ bool UUsdConversionBlueprintLibrary::HasMetadataField(
 
 	if (StageIdentifier.IsEmpty() && AssetUserData->StageIdentifierToMetadata.Num() != 1)
 	{
-		UE_LOG(
-			LogUsd,
-			Warning,
-			TEXT("Failed to check for metadata field '%s' on AssetUserData '%s': Please provide a valid value for the StageIdentifier parameter"),
-			*AssetUserData->GetPathName(),
-			*Key
-		);
+		USD_LOG_USERWARNING(FText::Format(
+			LOCTEXT(
+				"InvalidStageIdentifierCheck",
+				"Failed to check for metadata field '{0}' on AssetUserData '{1}': Please provide a valid value for the StageIdentifier parameter"
+			),
+			FText::FromString(AssetUserData->GetPathName()),
+			FText::FromString(Key)
+		));
 		return false;
 	}
 
@@ -1039,13 +1054,14 @@ bool UUsdConversionBlueprintLibrary::HasMetadataField(
 
 	if (PrimPath.IsEmpty() && CombinedPrimMetadata->PrimPathToMetadata.Num() != 1)
 	{
-		UE_LOG(
-			LogUsd,
-			Warning,
-			TEXT("Failed to clear metadata field '%s' on AssetUserData '%s': Please provide a valid value for the PrimPath parameter"),
-			*AssetUserData->GetPathName(),
-			*Key
-		);
+		USD_LOG_USERWARNING(FText::Format(
+			LOCTEXT(
+				"InvalidPrimPathCheck",
+				"Failed to clear metadata field '{0}' on AssetUserData '{1}': Please provide a valid value for the PrimPath parameter"
+			),
+			FText::FromString(AssetUserData->GetPathName()),
+			FText::FromString(Key)
+		));
 		return false;
 	}
 
@@ -1073,13 +1089,14 @@ FUsdMetadataValue UUsdConversionBlueprintLibrary::GetMetadataField(
 
 	if (StageIdentifier.IsEmpty() && AssetUserData->StageIdentifierToMetadata.Num() != 1)
 	{
-		UE_LOG(
-			LogUsd,
-			Warning,
-			TEXT("Failed to get metadata field '%s' on AssetUserData '%s': Please provide a valid value for the StageIdentifier parameter"),
-			*AssetUserData->GetPathName(),
-			*Key
-		);
+		USD_LOG_USERWARNING(FText::Format(
+			LOCTEXT(
+				"InvalidStageIdentifierGet",
+				"Failed to get metadata field '{0}' on AssetUserData '{1}': Please provide a valid value for the StageIdentifier parameter"
+			),
+			FText::FromString(AssetUserData->GetPathName()),
+			FText::FromString(Key)
+		));
 		return {};
 	}
 
@@ -1093,13 +1110,14 @@ FUsdMetadataValue UUsdConversionBlueprintLibrary::GetMetadataField(
 
 	if (PrimPath.IsEmpty() && CombinedPrimMetadata->PrimPathToMetadata.Num() != 1)
 	{
-		UE_LOG(
-			LogUsd,
-			Warning,
-			TEXT("Failed to get metadata field '%s' on AssetUserData '%s': Please provide a valid value for the PrimPath parameter"),
-			*AssetUserData->GetPathName(),
-			*Key
-		);
+		USD_LOG_USERWARNING(FText::Format(
+			LOCTEXT(
+				"InvalidPrimPathGet",
+				"Failed to get metadata field '{0}' on AssetUserData '{1}': Please provide a valid value for the PrimPath parameter"
+			),
+			FText::FromString(AssetUserData->GetPathName()),
+			FText::FromString(Key)
+		));
 		return {};
 	}
 
@@ -1115,15 +1133,13 @@ FUsdMetadataValue UUsdConversionBlueprintLibrary::GetMetadataField(
 		return *FoundValue;
 	}
 
-	UE_LOG(
-		LogUsd,
-		Warning,
-		TEXT("Failed to find a metadata entry in '%s' with stage identifier '%s', prim path '%s' and key '%s'"),
-		*AssetUserData->GetPathName(),
-		*StageIdentifier,
-		*PrimPath,
-		*Key
-	);
+	USD_LOG_USERWARNING(FText::Format(
+		LOCTEXT("FailedToFindMetadata", "Failed to find a metadata entry in '{0}' with stage identifier '{1}', prim path '{2}' and key '{3}'"),
+		FText::FromString(AssetUserData->GetPathName()),
+		FText::FromString(StageIdentifier),
+		FText::FromString(PrimPath),
+		FText::FromString(Key)
+	));
 	return {};
 }
 
@@ -1471,7 +1487,7 @@ namespace UE::UsdConversionLibrary::Private
 		bool bSuccess = Unstringifier(String, Output);
 		if (!bSuccess)
 		{
-			UE_LOG(LogUsd, Warning, TEXT("Failed to unstringify '%s' with the type '%s'!"), *String, *TypeName);
+			USD_LOG_WARNING(TEXT("Failed to unstringify '%s' with the type '%s'!"), *String, *TypeName);
 		}
 
 		return Output;
@@ -1825,3 +1841,5 @@ TArray<FIntVector4> UUsdConversionBlueprintLibrary::UnstringifyAsInt4Array(const
 {
 	return UE::UsdConversionLibrary::Private::UnstringifyOrDefault<TArray<FIntVector4>, UsdUtils::UnstringifyAsInt4Array>(String, TEXT("int4"));
 }
+
+#undef LOCTEXT_NAMESPACE
diff --git a/Engine/Plugins/Importers/USDImporter/Source/USDExporter/Private/USDExporterModule.cpp b/Engine/Plugins/Importers/USDImporter/Source/USDExporter/Private/USDExporterModule.cpp
--- a/Engine/Plugins/Importers/USDImporter/Source/USDExporter/Private/USDExporterModule.cpp
+++ b/Engine/Plugins/Importers/USDImporter/Source/USDExporter/Private/USDExporterModule.cpp
@@ -4,7 +4,7 @@
 
 #include "LevelExporterUSDOptionsCustomization.h"
 #include "USDAssetOptions.h"
-#include "USDLog.h"
+#include "USDErrorUtils.h"
 #include "USDMemory.h"
 
 #include "UsdWrappers/SdfLayer.h"
@@ -159,12 +159,11 @@ bool IUsdExporterModule::CanExportToLayer(const FString& TargetFilePath)
 			FText ErrorMessage = FText::Format(
 				LOCTEXT(
 					"FailedExportLayerOpenSubText",
-					"Failed to export to the USD layer '{0}' as a layer with this identifier is already open in another USD Stage. Please use a "
-					"different file path or try closing that other USD stage and exporting again."
+					"Failed to export to the USD layer '{0}' as a layer with this identifier is already open in another USD Stage. Please use a different file path or try closing that other USD stage and exporting again."
 				),
 				FText::FromString(TargetFilePath)
 			);
-			UE_LOG(LogUsd, Error, TEXT("%s"), *ErrorMessage.ToString());
+			USD_LOG_USERERROR(ErrorMessage);
 
 			FNotificationInfo ErrorToast(LOCTEXT("FailedExportLayerOpenText", "USD: Export failure"));
 			ErrorToast.ExpireDuration = 10.0f;
diff --git a/Engine/Plugins/Importers/USDImporter/Source/USDSchemas/Private/Custom/MaterialXUSDShadeMaterialTranslator.cpp b/Engine/Plugins/Importers/USDImporter/Source/USDSchemas/Private/Custom/MaterialXUSDShadeMaterialTranslator.cpp
--- a/Engine/Plugins/Importers/USDImporter/Source/USDSchemas/Private/Custom/MaterialXUSDShadeMaterialTranslator.cpp
+++ b/Engine/Plugins/Importers/USDImporter/Source/USDSchemas/Private/Custom/MaterialXUSDShadeMaterialTranslator.cpp
@@ -7,7 +7,6 @@
 #include "USDClassesModule.h"
 #include "USDConversionUtils.h"
 #include "USDErrorUtils.h"
-#include "USDLog.h"
 #include "USDMemory.h"
 #include "USDObjectUtils.h"
 #include "USDShadeConversion.h"
@@ -181,8 +180,6 @@ namespace UE::USDMaterialXTranslator::Private
 				const TArray<UObject*>& ImportedObjects
 			)
 			{
-				UE_LOG(LogUsd, Log, TEXT("Translated '%d' assets from MaterialX file '%s'"), ImportedObjects.Num(), *MaterialXFilePath);
-
 				// Move all our assets to the transient package
 				for (UObject* ImportedObject : ImportedObjects)
 				{
@@ -342,12 +339,6 @@ void FMaterialXUsdShadeMaterialTranslator::CreateAssets()
 
 	if (Context->bTranslateOnlyUsedMaterials && !Context->UsdInfoCache->IsMaterialUsed(PrimPath))
 	{
-		UE_LOG(
-			LogUsd,
-			Verbose,
-			TEXT("Skipping creating assets for material prim '%s' as it is not currently bound by any prim."),
-			*PrimPath.GetString()
-		);
 		return;
 	}
 
@@ -413,14 +404,13 @@ void FMaterialXUsdShadeMaterialTranslator::CreateAssets()
 
 	if (MaterialXFilePaths.IsEmpty() || !MaterialXReferencerPrim)
 	{
-		UE_LOG(
-			LogUsd,
-			Warning,
-			TEXT(
-				"Recognized potential MaterialX materials on prim '%s', but failed to find a valid referenced MaterialX file. Reverting to parsing the generated Material prims instead."
+		USD_LOG_USERWARNING(FText::Format(
+			LOCTEXT(
+				"NoReferencedMtlxFile",
+				"Recognized potential MaterialX materials on prim '{0}', but failed to find a valid referenced MaterialX file. Reverting to parsing the generated Material prims instead."
 			),
-			*PrimPath.GetString()
-		);
+			FText::FromString(PrimPath.GetString())
+		));
 		Super::CreateAssets();
 		return;
 	}
@@ -471,13 +461,14 @@ void FMaterialXUsdShadeMaterialTranslator::CreateAssets()
 
 			if (!bSuccess)
 			{
-				UE_LOG(
-					LogUsd,
-					Warning,
-					TEXT("Recognized potential MaterialX materials on prim '%s', but MaterialX parsing of file '%s' failed."),
-					*PrimPath.GetString(),
-					*MaterialXFilePath
-				);
+				USD_LOG_USERWARNING(FText::Format(
+					LOCTEXT(
+						"ParsingReferencedFailed",
+						"Recognized potential MaterialX materials on prim '{0}', but MaterialX parsing of file '{1}' failed."
+					),
+					FText::FromString(PrimPath.GetString()),
+					FText::FromString(MaterialXFilePath)
+				));
 			}
 
 			// Check if we found our target material when parsing this file
@@ -509,14 +500,14 @@ void FMaterialXUsdShadeMaterialTranslator::CreateAssets()
 		}
 		FilePaths.RemoveFromEnd(TEXT(", "));
 
-		UE_LOG(
-			LogUsd,
-			Warning,
-			TEXT("Failed to find target Material '%s' after parsing all Material X files [%s]. Reverting back to parsing the USD Material prim "
-				 "generated by usdMtlx directly."),
-			*PrimPath.GetString(),
-			*FilePaths
-		);
+		USD_LOG_USERWARNING(FText::Format(
+			LOCTEXT(
+				"MaterialNotFoundInMtlxFile",
+				"Failed to find target Material '{0}' after parsing all Material X files [{1}]. Reverting back to parsing the USD Material prim generated by usdMtlx directly."
+			),
+			FText::FromString(PrimPath.GetString()),
+			FText::FromString(FilePaths)
+		));
 		Super::CreateAssets();
 	}
 #else
diff --git a/Engine/Plugins/Importers/USDImporter/Source/USDSchemas/Private/MeshTranslationImpl.cpp b/Engine/Plugins/Importers/USDImporter/Source/USDSchemas/Private/MeshTranslationImpl.cpp
--- a/Engine/Plugins/Importers/USDImporter/Source/USDSchemas/Private/MeshTranslationImpl.cpp
+++ b/Engine/Plugins/Importers/USDImporter/Source/USDSchemas/Private/MeshTranslationImpl.cpp
@@ -10,8 +10,8 @@
 #include "USDAssetUserData.h"
 #include "USDClassesModule.h"
 #include "USDConversionUtils.h"
+#include "USDErrorUtils.h"
 #include "USDGeomMeshConversion.h"
-#include "USDLog.h"
 #include "USDMemory.h"
 #include "USDObjectUtils.h"
 #include "USDProjectSettings.h"
@@ -40,6 +40,8 @@
 #include "pxr/usd/usdGeom/mesh.h"
 #include "USDIncludesEnd.h"
 
+#define LOCTEXT_NAMESPACE "MeshTranslationImpl"
+
 static_assert(
 	USD_PREVIEW_SURFACE_MAX_UV_SETS <= MAX_MESH_TEXTURE_COORDS_MD,
 	"UsdPreviewSurface materials can only have up to as many UV sets as MeshDescription supports!"
@@ -186,17 +188,6 @@ namespace UE::MeshTranslationImplInternal::Private
 			return nullptr;
 		}
 
-		UE_LOG(
-			LogUsd,
-			Verbose,
-			TEXT("Getting compatible material based on Material '%s' (Parameter to primvar: %s, primvar to UV index: %s) matching mesh primvar to UV "
-				 "index mapping '%s'"),
-			*Material.GetPathName(),
-			*UsdUtils::StringifyMap(MaterialAssetUserData->ParameterToPrimvar),
-			*UsdUtils::StringifyMap(MaterialAssetUserData->PrimvarToUVIndex),
-			*UsdUtils::StringifyMap(MeshPrimvarToUVIndex)
-		);
-
 		TArray<TSet<FString>> UVIndexToMeshPrimvars;
 		UVIndexToMeshPrimvars.SetNum(USD_PREVIEW_SURFACE_MAX_UV_SETS);
 		for (const TPair<FString, int32>& MeshPair : MeshPrimvarToUVIndex)
@@ -227,9 +218,7 @@ namespace UE::MeshTranslationImplInternal::Private
 			}
 			else
 			{
-				UE_LOG(
-					LogUsd,
-					Log,
+				USD_LOG_INFO(
 					TEXT("Failed to find primvar '%s' needed by material '%s' on its assigned mesh. Available primvars and UV indices: %s"),
 					*MaterialPrimvar,
 					*Material.GetPathName(),
@@ -251,7 +240,6 @@ namespace UE::MeshTranslationImplInternal::Private
 		}
 		if (bCompatible)
 		{
-			UE_LOG(LogUsd, Verbose, TEXT("Material '%s' is compatible with provided primvar to UV index mapping"), *Material.GetPathName());
 			return &Material;
 		}
 
@@ -386,16 +374,6 @@ namespace UE::MeshTranslationImplInternal::Private
 		UMaterialInstance* CompatibleInstance = Cast<UMaterialInstance>(CompatibleMaterial);
 		if (bCreatedNew && CompatibleInstance && CompatibleUserData)
 		{
-			UE_LOG(
-				LogUsd,
-				Verbose,
-				TEXT("Generated compatible version of Material '%s' (Parameter to primvar: %s, primvar to UV index: %s) with hash '%s'"),
-				*Material.GetPathName(),
-				*UsdUtils::StringifyMap(MaterialAssetUserData->ParameterToPrimvar),
-				*UsdUtils::StringifyMap(CompatiblePrimvarToUVIndex),
-				*PrefixedCompatibleHash
-			);
-
 			for (const TPair<FString, FString>& ParameterPair : CompatibleUserData->ParameterToPrimvar)
 			{
 				const FString& Parameter = ParameterPair.Key;
@@ -670,9 +648,7 @@ TMap<const UsdUtils::FUsdPrimMaterialSlot*, UMaterialInterface*> MeshTranslation
 						}
 						else
 						{
-							UE_LOG(
-								LogUsd,
-								Warning,
+							USD_LOG_WARNING(
 								TEXT(
 									"Failed to generate a two-sided material from the material prim at path '%s'. Falling back to using the single-sided material '%s' instead."
 								),
@@ -732,32 +708,28 @@ TMap<const UsdUtils::FUsdPrimMaterialSlot*, UMaterialInterface*> MeshTranslation
 					Material = Cast<UMaterialInterface>(Object);
 					if (!Object)
 					{
-						UE_LOG(
-							LogUsd,
-							Warning,
-							TEXT("UE material '%s' for prim '%s' could not be loaded or was not found."),
-							*Slot.MaterialSource,
-							*UsdToUnreal::ConvertPath(UsdPrim.GetPrimPath())
-						);
+						USD_LOG_USERWARNING(FText::Format(
+							LOCTEXT("FailToLoadMaterial", "UE material '{0}' for prim '{1}' could not be loaded or was not found."),
+							FText::FromString(Slot.MaterialSource),
+							FText::FromString(UsdToUnreal::ConvertPath(UsdPrim.GetPrimPath()))
+						));
 					}
 					else if (!Material)
 					{
-						UE_LOG(
-							LogUsd,
-							Warning,
-							TEXT("Object '%s' assigned as an Unreal Material for prim '%s' is not actually a material (but instead a '%s') and will "
-								 "not be used"),
-							*Slot.MaterialSource,
-							*UsdToUnreal::ConvertPath(UsdPrim.GetPrimPath()),
-							*Object->GetClass()->GetName()
-						);
+						USD_LOG_USERWARNING(FText::Format(
+							LOCTEXT(
+								"NotAMaterial",
+								"Object '{0}' assigned as an Unreal Material for prim '{1}' is not actually a material (but instead a '{2}') and will not be used"
+							),
+							FText::FromString(Slot.MaterialSource),
+							FText::FromString(UsdToUnreal::ConvertPath(UsdPrim.GetPrimPath())),
+							FText::FromString(Object->GetClass()->GetName())
+						));
 					}
 					else if (!Material->IsTwoSided() && Slot.bMeshIsDoubleSided)
 					{
-						UE_LOG(
-							LogUsd,
-							Warning,
-							TEXT("Using one-sided UE material '%s' for doubleSided prim '%s'"),
+						USD_LOG_WARNING(
+							TEXT("Using one-sided UE material '{0}' for doubleSided prim '{1}'"),
 							*Slot.MaterialSource,
 							*UsdToUnreal::ConvertPath(UsdPrim.GetPrimPath())
 						);
@@ -980,9 +952,7 @@ void MeshTranslationImpl::SetMaterialOverrides(
 			}
 			else
 			{
-				UE_LOG(
-					LogUsd,
-					Error,
+				USD_LOG_ERROR(
 					TEXT("Lost track of resolved material for slot '%d' of LOD '%d' for mesh '%s'"),
 					LODSlotIndex,
 					LODIndex,
@@ -1034,4 +1004,6 @@ void MeshTranslationImpl::RecordSourcePrimsForMaterialSlots(
 	}
 }
 
+#undef LOCTEXT_NAMESPACE
+
 #endif	  // #if USE_USD_SDK
diff --git a/Engine/Plugins/Importers/USDImporter/Source/USDSchemas/Private/USDGeomMeshTranslator.cpp b/Engine/Plugins/Importers/USDImporter/Source/USDSchemas/Private/USDGeomMeshTranslator.cpp
--- a/Engine/Plugins/Importers/USDImporter/Source/USDSchemas/Private/USDGeomMeshTranslator.cpp
+++ b/Engine/Plugins/Importers/USDImporter/Source/USDSchemas/Private/USDGeomMeshTranslator.cpp
@@ -13,8 +13,8 @@
 #include "USDClassesModule.h"
 #include "USDConversionUtils.h"
 #include "USDDrawModeComponent.h"
+#include "USDErrorUtils.h"
 #include "USDGeomMeshConversion.h"
-#include "USDLog.h"
 #include "USDMemory.h"
 #include "USDObjectUtils.h"
 #include "USDPrimConversion.h"
@@ -64,6 +64,8 @@
 #include "pxr/usd/usdSkel/bindingAPI.h"
 #include "USDIncludesEnd.h"
 
+#define LOCTEXT_NAMESPACE "USDGeomMeshTranslator"
+
 static bool GBuildReversedIndexBuffer = false;
 static FAutoConsoleVariableRef CVarBuildReversedIndexBuffer(
 	TEXT("USD.StaticMesh.BuildSettings.BuildReversedIndexBuffer"),
@@ -134,9 +136,7 @@ namespace UsdGeomMeshTranslatorImpl
 				if (OverrideValue == UnrealIdentifiers::UnrealNaniteOverrideEnable)
 				{
 					bHasNaniteOverrideEnabled = true;
-					UE_LOG(
-						LogUsd,
-						Log,
+					USD_LOG_INFO(
 						TEXT("Trying to enable Nanite for mesh generated for prim '%s' as the '%s' attribute is set to '%s'"),
 						*PrimPath.GetString(),
 						*UsdToUnreal::ConvertToken(UnrealIdentifiers::UnrealNaniteOverride),
@@ -145,9 +145,7 @@ namespace UsdGeomMeshTranslatorImpl
 				}
 				else if (OverrideValue == UnrealIdentifiers::UnrealNaniteOverrideDisable)
 				{
-					UE_LOG(
-						LogUsd,
-						Log,
+					USD_LOG_INFO(
 						TEXT("Not enabling Nanite for mesh generated for prim '%s' as the '%s' attribute is set to '%s'"),
 						*PrimPath.GetString(),
 						*UsdToUnreal::ConvertToken(UnrealIdentifiers::UnrealNaniteOverride),
@@ -162,27 +160,8 @@ namespace UsdGeomMeshTranslatorImpl
 		if (!bHasNaniteOverrideEnabled)
 		{
 			const int32 NumTriangles = LODIndexToMeshDescription[0].Triangles().Num();
-			if (NumTriangles >= Context.NaniteTriangleThreshold)
+			if (NumTriangles < Context.NaniteTriangleThreshold)
 			{
-				UE_LOG(
-					LogUsd,
-					Verbose,
-					TEXT("Trying to enable Nanite for mesh generated for prim '%s' as it has '%d' triangles, and the threshold is '%d'"),
-					*PrimPath.GetString(),
-					NumTriangles,
-					Context.NaniteTriangleThreshold
-				);
-			}
-			else
-			{
-				UE_LOG(
-					LogUsd,
-					Verbose,
-					TEXT("Not enabling Nanite for mesh generated for prim '%s' as it has '%d' triangles, and the threshold is '%d'"),
-					*PrimPath.GetString(),
-					NumTriangles,
-					Context.NaniteTriangleThreshold
-				);
 				return false;
 			}
 		}
@@ -193,26 +172,24 @@ namespace UsdGeomMeshTranslatorImpl
 
 			if (!SubtreeSectionCount.IsSet() || SubtreeSectionCount.GetValue() > NANITE_MAX_CLUSTER_MATERIALS)
 			{
-				UE_LOG(
-					LogUsd,
-					Warning,
-					TEXT("Not enabling Nanite for mesh generated for prim '%s' as LOD0 has '%d' material slots, which is above the Nanite limit of "
-						 "'%d'"),
-					*PrimPath.GetString(),
+				USD_LOG_USERWARNING(FText::Format(
+					LOCTEXT(
+						"NaniteSlotWarning",
+						"Not enabling Nanite for mesh generated for prim '{0}' as LOD0 has '{1}' material slots, which is above the Nanite limit of '{2}'"
+					),
+					FText::FromString(PrimPath.GetString()),
 					SubtreeSectionCount.GetValue(),
 					NANITE_MAX_CLUSTER_MATERIALS
-				);
+				));
 				return false;
 			}
 		}
 
 #if !WITH_EDITOR
-		UE_LOG(
-			LogUsd,
-			Warning,
-			TEXT("Not enabling Nanite for mesh generated for prim '%s' as we can't setup Nanite during runtime"),
-			*PrimPath.GetString()
-		);
+		USD_LOG_USERWARNING(FText::Format(
+			LOCTEXT("NaniteSlotWarning", "Not enabling Nanite for mesh generated for prim '{0}' as we can't setup Nanite during runtime"),
+			FText::FromString(PrimPath.GetString())
+		));
 		return false;
 #else
 		return true;
@@ -264,9 +241,7 @@ namespace UsdGeomMeshTranslatorImpl
 				}
 				else
 				{
-					UE_LOG(
-						LogUsd,
-						Error,
+					USD_LOG_ERROR(
 						TEXT("Failed to resolve material '%s' for slot '%d' of LOD '%d' for mesh '%s'"),
 						*Slot.MaterialSource,
 						LODSlotIndex,
@@ -1703,11 +1678,10 @@ void FBuildStaticMeshTaskChain::SetupTasks()
 			// Build failed, abandon mesh
 			if (!bSuccess)
 			{
-				UE_LOG(
-					LogUsd,
-					Warning,
-					TEXT("Discarding StaticMesh generated for prim '%s' as it didn't produce any valid RenderData (likely all triangles "
-						 "were degenerate)"),
+				USD_LOG_WARNING(
+					TEXT(
+						"Discarding StaticMesh generated for prim '%s' as it didn't produce any valid RenderData (likely all triangles were degenerate)"
+					),
 					*PrimPath.GetString()
 				);
 				UsdUnreal::TranslatorUtils::AbandonFailedAsset(StaticMesh, Context->UsdAssetCache.Get(), Context->PrimLinkCache);
@@ -1825,13 +1799,11 @@ void FGeomMeshCreateAssetsTaskChain::SetupTasks()
 					   LODIndexToMaterialInfo.RemoveAt(LODIndex, EAllowShrinking::No);
 				   }
 
-				   UE_LOG(
-					   LogUsd,
-					   Warning,
-					   TEXT("Ignoring mesh data collected for LOD%d of prim '%s' as it is empty. Is the prim invisible?"),
+				   USD_LOG_USERWARNING(FText::Format(
+					   LOCTEXT("EmpytMeshData", "Ignoring mesh data collected for LOD{0} of prim '{1}' as it is empty. Is the prim invisible?"),
 					   LODIndex,
-					   bParseLODs ? *GetPrim().GetParent().GetPrimPath().GetString() : *GetPrim().GetPrimPath().GetString()
-				   );
+					   FText::FromString(bParseLODs ? *GetPrim().GetParent().GetPrimPath().GetString() : *GetPrim().GetPrimPath().GetString())
+				   ));
 			   }
 			   else
 			   {
@@ -2163,4 +2135,6 @@ bool FUsdGeomMeshTranslator::ShouldSkipSkinnablePrim(bool bCheckForComponent) co
 	return false;
 }
 
+#undef LOCTEXT_NAMESPACE
+
 #endif	  // #if USE_USD_SDK
diff --git a/Engine/Plugins/Importers/USDImporter/Source/USDSchemas/Private/USDGeomPointInstancerTranslator.cpp b/Engine/Plugins/Importers/USDImporter/Source/USDSchemas/Private/USDGeomPointInstancerTranslator.cpp
--- a/Engine/Plugins/Importers/USDImporter/Source/USDSchemas/Private/USDGeomPointInstancerTranslator.cpp
+++ b/Engine/Plugins/Importers/USDImporter/Source/USDSchemas/Private/USDGeomPointInstancerTranslator.cpp
@@ -8,9 +8,9 @@
 #include "USDAssetUserData.h"
 #include "USDConversionUtils.h"
 #include "USDDrawModeComponent.h"
+#include "USDErrorUtils.h"
 #include "USDGeomMeshConversion.h"
 #include "USDGeomMeshTranslator.h"
-#include "USDLog.h"
 #include "USDMemory.h"
 #include "USDPrimConversion.h"
 #include "USDTypesConversion.h"
@@ -210,9 +210,7 @@ void FUsdGeomPointInstancerTranslator::CreateAssets()
 
 			if (!PrototypeUsdPrim)
 			{
-				UE_LOG(
-					LogUsd,
-					Warning,
+				USD_LOG_WARNING(
 					TEXT("Failed to find prototype '%s' for PointInstancer '%s' when collapsing assets"),
 					*UEPrototypePath.GetString(),
 					*PrimPath.GetString()
@@ -346,9 +344,7 @@ USceneComponent* FUsdGeomPointInstancerTranslator::CreateComponents()
 			pxr::UsdPrim PrototypeUsdPrim = Prim.GetStage()->GetPrimAtPath(PrototypePath);
 			if (!PrototypeUsdPrim)
 			{
-				UE_LOG(
-					LogUsd,
-					Warning,
+				USD_LOG_WARNING(
 					TEXT("Failed to find prototype '%s' for PointInstancer '%s' when creating components"),
 					*PrototypePathStr,
 					*PrimPath.GetString()
@@ -385,9 +381,7 @@ USceneComponent* FUsdGeomPointInstancerTranslator::CreateComponents()
 
 			if (!ISMComponent)
 			{
-				UE_LOG(
-					LogUsd,
-					Warning,
+				USD_LOG_WARNING(
 					TEXT("Failed to generate ISM/HISM component for prototype '%s' for PointInstancer '%s'"),
 					*PrototypePathStr,
 					*PrimPath.GetString()
@@ -484,9 +478,7 @@ void FUsdGeomPointInstancerTranslator::UpdateComponents(USceneComponent* PointIn
 			TUsdStore<pxr::UsdPrim> PrototypeUsdPrim = Prim.GetStage()->GetPrimAtPath(PrototypePath);
 			if (!PrototypeUsdPrim.Get())
 			{
-				UE_LOG(
-					LogUsd,
-					Warning,
+				USD_LOG_WARNING(
 					TEXT("Failed to find prototype '%s' for PointInstancer '%s' when updating components"),
 					*UsdToUnreal::ConvertPath(PrototypePath),
 					*PrimPath.GetString()
@@ -499,9 +491,7 @@ void FUsdGeomPointInstancerTranslator::UpdateComponents(USceneComponent* PointIn
 			// The user could have just manually deleted the component, so we must check
 			if (!AttachedISMs.IsValidIndex(PrototypeIndex))
 			{
-				UE_LOG(
-					LogUsd,
-					Warning,
+				USD_LOG_WARNING(
 					TEXT("Failed to find corresponding ISM/HISM component for prototype '%s' of PointInstancer '%s'. Cancelling component update"),
 					*UsdToUnreal::ConvertPath(PrototypePath),
 					*PrimPath.GetString()
diff --git a/Engine/Plugins/Importers/USDImporter/Source/USDSchemas/Private/USDGeomXformableTranslator.cpp b/Engine/Plugins/Importers/USDImporter/Source/USDSchemas/Private/USDGeomXformableTranslator.cpp
--- a/Engine/Plugins/Importers/USDImporter/Source/USDSchemas/Private/USDGeomXformableTranslator.cpp
+++ b/Engine/Plugins/Importers/USDImporter/Source/USDSchemas/Private/USDGeomXformableTranslator.cpp
@@ -9,10 +9,10 @@
 #include "USDAssetUserData.h"
 #include "USDConversionUtils.h"
 #include "USDDrawModeComponent.h"
+#include "USDErrorUtils.h"
 #include "USDGeomMeshConversion.h"
 #include "USDGeomMeshTranslator.h"
 #include "USDIntegrationUtils.h"
-#include "USDLog.h"
 #include "USDMemory.h"
 #include "USDObjectUtils.h"
 #include "USDPrimConversion.h"
@@ -54,6 +54,8 @@
 #include "pxr/usd/usdSkel/root.h"
 #include "USDIncludesEnd.h"
 
+#define LOCTEXT_NAMESPACE "USDGeomXformableTranslator"
+
 static bool GCollapsePrimsWithoutKind = true;
 static FAutoConsoleVariableRef CVarCollapsePrimsWithoutKind(
 	TEXT("USD.CollapsePrimsWithoutKind"),
@@ -496,7 +498,7 @@ USceneComponent* FUsdGeomXformableTranslator::CreateComponentsEx(TOptional<TSubc
 
 	if (!ComponentOuter)
 	{
-		UE_LOG(LogUsd, Warning, TEXT("Invalid outer when trying to create SceneComponent for prim (%s)"), *PrimPath.GetString());
+		USD_LOG_WARNING(TEXT("Invalid outer when trying to create SceneComponent for prim (%s)"), *PrimPath.GetString());
 		return nullptr;
 	}
 
@@ -990,14 +992,15 @@ void FUsdGeomXformableTranslator::CreateAlternativeDrawModeAssets(EUsdDrawMode D
 				Attr.Get(&TextureAssetPath, Context->Time);
 				FString TargetAssetPath = UsdToUnreal::ConvertString(TextureAssetPath.GetAssetPath());
 
-				UE_LOG(
-					LogUsd,
-					Warning,
-					TEXT("Failed to resolve texture path '%s' specified as a card geometry texture on prim '%s' at time '%f'"),
-					*TargetAssetPath,
-					*UsdToUnreal::ConvertPath(Attr.GetPrim().GetPrimPath()),
+				USD_LOG_USERWARNING(FText::Format(
+					LOCTEXT(
+						"CardTextureFailResolve",
+						"Failed to resolve texture path '{0}' specified as a card geometry texture on prim '{1}' at time '{2}'"
+					),
+					FText::FromString(TargetAssetPath),
+					FText::FromString(UsdToUnreal::ConvertPath(Attr.GetPrim().GetPrimPath())),
 					Context->Time
-				);
+				));
 			}
 			else
 			{
@@ -1198,4 +1201,6 @@ TSet<UE::FSdfPath> FUsdGeomXformableTranslator::CollectAuxiliaryPrims() const
 	return Result;
 }
 
+#undef LOCTEXT_NAMESPACE
+
 #endif	  // #if USE_USD_SDK
diff --git a/Engine/Plugins/Importers/USDImporter/Source/USDSchemas/Private/USDGeometryCacheTranslator.cpp b/Engine/Plugins/Importers/USDImporter/Source/USDSchemas/Private/USDGeometryCacheTranslator.cpp
--- a/Engine/Plugins/Importers/USDImporter/Source/USDSchemas/Private/USDGeometryCacheTranslator.cpp
+++ b/Engine/Plugins/Importers/USDImporter/Source/USDSchemas/Private/USDGeometryCacheTranslator.cpp
@@ -12,9 +12,9 @@
 #include "USDAssetUserData.h"
 #include "USDConversionUtils.h"
 #include "USDDrawModeComponent.h"
+#include "USDErrorUtils.h"
 #include "USDGroomTranslatorUtils.h"
 #include "USDIntegrationUtils.h"
-#include "USDLog.h"
 #include "USDMemory.h"
 #include "USDObjectUtils.h"
 #include "USDPrimConversion.h"
@@ -117,9 +117,7 @@ namespace UsdGeometryCacheTranslatorImpl
 				else
 				{
 					// Warn, but still add a material slot to preserve the materials order
-					UE_LOG(
-						LogUsd,
-						Warning,
+					USD_LOG_WARNING(
 						TEXT("Failed to resolve material '%s' for slot '%d' for geometry cache '%s'"),
 						*Slot.MaterialSource,
 						LODSlotIndex,
@@ -522,7 +520,7 @@ namespace UsdGeometryCacheTranslatorImpl
 		const int32 NumFrames = (Args.EndFrame - Args.StartFrame) * NumMeshes;
 		int32 TotalFrames = 0;
 
-		FString Title = FString::Format(TEXT("Importing frames for {0} ({1}/{2})"), { AssetName, 0,  NumFrames });
+		FString Title = FString::Format(TEXT("Importing frames for {0} ({1}/{2})"), {AssetName, 0, NumFrames});
 		FScopedSlowTask SlowTask(static_cast<float>(NumFrames), FText::FromString(Title), true, GeoCacheContext);
 		SlowTask.MakeDialog(true);
 
@@ -551,7 +549,7 @@ namespace UsdGeometryCacheTranslatorImpl
 					StreamableTrack->AddMeshSample(MeshData, (FrameIndex - Args.StartFrame) / Args.FramesPerSecond, bConstantTopology);
 
 					++TotalFrames;
-					Title = FString::Format(TEXT("Importing frames for {0} ({1}/{2})"), { AssetName, TotalFrames,  NumFrames });
+					Title = FString::Format(TEXT("Importing frames for {0} ({1}/{2})"), {AssetName, TotalFrames, NumFrames});
 					SlowTask.EnterProgressFrame(1, FText::FromString(Title));
 				}
 			}
@@ -590,13 +588,20 @@ namespace UsdGeometryCacheTranslatorImpl
 			bool bCancelled = false;
 			ParallelFor(
 				NumMeshThreads,
-				[&Args, &MeshPrims, &MaterialOffsets, &SyncEvents, NumMeshThreads, NumFrameThreads, GeometryCache, NumMeshes,
-				&SlowTask,
-				&AssetName,
-				&CompletedFrames,
-				&NumFrames,
-				&TotalFrames,
-				&bCancelled](int32 MeshThreadIndex)
+				[&Args,
+				 &MeshPrims,
+				 &MaterialOffsets,
+				 &SyncEvents,
+				 NumMeshThreads,
+				 NumFrameThreads,
+				 GeometryCache,
+				 NumMeshes,
+				 &SlowTask,
+				 &AssetName,
+				 &CompletedFrames,
+				 &NumFrames,
+				 &TotalFrames,
+				 &bCancelled](int32 MeshThreadIndex)
 				{
 					int32 MeshIndex = MeshThreadIndex;
 
@@ -673,7 +678,10 @@ namespace UsdGeometryCacheTranslatorImpl
 											float TempCompletedFrames = CompletedFrames;
 											CompletedFrames = 0.0f;
 											TotalFrames += int32(TempCompletedFrames);
-											FString Title = FString::Format(TEXT("Importing frames for {0} ({1}/{2})"), { AssetName, TotalFrames,  NumFrames });
+											FString Title = FString::Format(
+												TEXT("Importing frames for {0} ({1}/{2})"),
+												{AssetName, TotalFrames, NumFrames}
+											);
 											SlowTask.EnterProgressFrame(TempCompletedFrames, FText::FromString(Title));
 										}
 
@@ -700,7 +708,7 @@ namespace UsdGeometryCacheTranslatorImpl
 
 	void FinalizeGeometryCache(UGeometryCache* GeometryCache)
 	{
-		for (int32 Index = GeometryCache->Tracks.Num() - 1; Index >= 0 ; --Index)
+		for (int32 Index = GeometryCache->Tracks.Num() - 1; Index >= 0; --Index)
 		{
 			if (UGeometryCacheTrackStreamable* StreamableTrack = Cast<UGeometryCacheTrackStreamable>(GeometryCache->Tracks[Index]))
 			{
diff --git a/Engine/Plugins/Importers/USDImporter/Source/USDSchemas/Private/USDGroomTranslator.cpp b/Engine/Plugins/Importers/USDImporter/Source/USDSchemas/Private/USDGroomTranslator.cpp
--- a/Engine/Plugins/Importers/USDImporter/Source/USDSchemas/Private/USDGroomTranslator.cpp
+++ b/Engine/Plugins/Importers/USDImporter/Source/USDSchemas/Private/USDGroomTranslator.cpp
@@ -22,10 +22,10 @@
 #include "USDAssetUserData.h"
 #include "USDClassesModule.h"
 #include "USDDrawModeComponent.h"
+#include "USDErrorUtils.h"
 #include "USDGroomConversion.h"
 #include "USDGroomTranslatorUtils.h"
 #include "USDIntegrationUtils.h"
-#include "USDLog.h"
 #include "USDMemory.h"
 #include "USDObjectUtils.h"
 #include "USDPrimConversion.h"
@@ -38,6 +38,8 @@
 #include "pxr/usd/usdGeom/imageable.h"
 #include "USDIncludesEnd.h"
 
+#define LOCTEXT_NAMESPACE "USDGroomTranslator"
+
 namespace UE::UsdGroomTranslator::Private
 {
 	UGroomImportOptions* CreateGroomImportOptions(
@@ -402,15 +404,15 @@ protected:
 							if (CacheInputData[GroupIndex].Strands.GetNumPoints() != GroomHairGroupsData[GroupIndex].Strands.BulkData.GetNumPoints())
 							{
 								bSuccess = false;
-								UE_LOG(
-									LogUsd,
-									Warning,
-									TEXT("GroomCache frame %d does not have the same number of vertices as the static groom (%u instead of %u). "
-										 "Aborting GroomCache import."),
+								USD_LOG_USERWARNING(FText::Format(
+									LOCTEXT(
+										"GroomCacheWrongNumVertices",
+										"GroomCache frame {0} does not have the same number of vertices as the static groom ({1} instead of {2}). Aborting GroomCache import."
+									),
 									FrameIndex,
 									CacheInputData[GroupIndex].Strands.GetNumPoints(),
 									GroomHairGroupsData[GroupIndex].Strands.BulkData.GetNumPoints()
-								);
+								));
 								break;
 							}
 						}
@@ -418,14 +420,14 @@ protected:
 					else
 					{
 						bSuccess = false;
-						UE_LOG(
-							LogUsd,
-							Warning,
-							TEXT("GroomCache does not have the same number of groups as the static groom (%d instead of %d). Aborting GroomCache "
-								 "import."),
+						USD_LOG_USERWARNING(FText::Format(
+							LOCTEXT(
+								"GroomCacheWrongNumGroups",
+								"GroomCache does not have the same number of groups as the static groom ({0} instead of {1}). Aborting GroomCache import."
+							),
 							CacheInputData.Num(),
 							GroomHairGroupsData.Num()
-						);
+						));
 					}
 
 					if (!bSuccess)
@@ -449,7 +451,7 @@ protected:
 				else
 				{
 					const FString StrandsGroomCachePrimPath = UsdGroomTranslatorUtils::GetStrandsGroomCachePrimPath(PrimPath);
-					UE_LOG(LogUsd, Warning, TEXT("Failed to create GroomCache for prim '%s'"), *StrandsGroomCachePrimPath);
+					USD_LOG_WARNING(TEXT("Failed to create GroomCache for prim '%s'"), *StrandsGroomCachePrimPath);
 
 					UsdUnreal::TranslatorUtils::AbandonFailedAsset(GroomCache, Context->UsdAssetCache.Get(), Context->PrimLinkCache);
 				}
@@ -649,4 +651,6 @@ TSet<UE::FSdfPath> FUsdGroomTranslator::CollectAuxiliaryPrims() const
 	return Result;
 }
 
+#undef LOCTEXT_NAMESPACE
+
 #endif	  // #if USE_USD_SDK
diff --git a/Engine/Plugins/Importers/USDImporter/Source/USDSchemas/Private/USDLuxLightTranslator.cpp b/Engine/Plugins/Importers/USDImporter/Source/USDSchemas/Private/USDLuxLightTranslator.cpp
--- a/Engine/Plugins/Importers/USDImporter/Source/USDSchemas/Private/USDLuxLightTranslator.cpp
+++ b/Engine/Plugins/Importers/USDImporter/Source/USDSchemas/Private/USDLuxLightTranslator.cpp
@@ -6,8 +6,8 @@
 #include "USDAssetUserData.h"
 #include "USDConversionUtils.h"
 #include "USDDrawModeComponent.h"
+#include "USDErrorUtils.h"
 #include "USDLightConversion.h"
-#include "USDLog.h"
 #include "USDMemory.h"
 #include "USDObjectUtils.h"
 #include "USDShadeConversion.h"
@@ -30,6 +30,8 @@
 #include "pxr/usd/usdLux/rectLight.h"
 #include "USDIncludesEnd.h"
 
+#define LOCTEXT_NAMESPACE "USDLuxLightTranslator"
+
 void FUsdLuxLightTranslator::CreateAssets()
 {
 	TRACE_CPUPROFILER_EVENT_SCOPE(FUsdLuxLightTranslator::CreateAssets);
@@ -67,13 +69,11 @@ void FUsdLuxLightTranslator::CreateAssets()
 		// Show a good warning for this because it's easy to pick some cubemap asset from the engine (that usually don't come with the
 		// source texture) and have the dome light silently not work again
 		FString TargetAssetPath = UsdToUnreal::ConvertString(TextureAssetPath.GetAssetPath());
-		UE_LOG(
-			LogUsd,
-			Warning,
-			TEXT("Failed to find texture '%s' used for UsdLuxDomeLight '%s'!"),
-			*TargetAssetPath,
-			*UsdToUnreal::ConvertPath(DomeLight.GetPrim().GetPath())
-		);
+		USD_LOG_USERWARNING(FText::Format(
+			LOCTEXT("FailToFindTexture", "Failed to find texture '{0}' used for UsdLuxDomeLight '{1}'!"),
+			FText::FromString(TargetAssetPath),
+			FText::FromString(UsdToUnreal::ConvertPath(DomeLight.GetPrim().GetPath()))
+		));
 
 		return;
 	}
@@ -214,4 +214,6 @@ bool FUsdLuxLightTranslator::CanBeCollapsed(ECollapsingType CollapsingType) cons
 	return false;
 }
 
+#undef LOCTEXT_NAMESPACE
+
 #endif	  // #if USE_USD_SDK
diff --git a/Engine/Plugins/Importers/USDImporter/Source/USDSchemas/Private/USDMediaSpatialAudioTranslator.cpp b/Engine/Plugins/Importers/USDImporter/Source/USDSchemas/Private/USDMediaSpatialAudioTranslator.cpp
--- a/Engine/Plugins/Importers/USDImporter/Source/USDSchemas/Private/USDMediaSpatialAudioTranslator.cpp
+++ b/Engine/Plugins/Importers/USDImporter/Source/USDSchemas/Private/USDMediaSpatialAudioTranslator.cpp
@@ -6,7 +6,7 @@
 #include "USDAssetUserData.h"
 #include "USDConversionUtils.h"
 #include "USDDrawModeComponent.h"
-#include "USDLog.h"
+#include "USDErrorUtils.h"
 #include "USDMemory.h"
 #include "USDObjectUtils.h"
 #include "USDPrimConversion.h"
@@ -35,6 +35,8 @@
 #include <pxr/usd/usdMedia/spatialAudio.h>
 #include "USDIncludesEnd.h"
 
+#define LOCTEXT_NAMESPACE "USDMediaSpatialAudioTranslator"
+
 namespace UE::UsdMediaSpatialAudioTranslator::Private
 {
 	USoundWave* CreateSoundWave(const FString& FilePath, FName AssetName, UObject* Outer, EObjectFlags Flags)
@@ -70,7 +72,7 @@ namespace UE::UsdMediaSpatialAudioTranslator::Private
 		TArray<uint8> FileBytes;
 		if (!FFileHelper::LoadFileToArray(FileBytes, *FilePath))
 		{
-			UE_LOG(LogUsd, Warning, TEXT("Failed to load audio file '%s' into array"), *FilePath);
+			USD_LOG_WARNING(TEXT("Failed to load audio file '%s' into array"), *FilePath);
 			return nullptr;
 		}
 
@@ -93,7 +95,7 @@ namespace UE::UsdMediaSpatialAudioTranslator::Private
 			}
 			else
 			{
-				UE_LOG(LogUsd, Warning, TEXT("Failed to convert audio file format '%s' to wav"), *Extension);
+				USD_LOG_WARNING(TEXT("Failed to convert audio file format '%s' to wav"), *Extension);
 				return nullptr;
 			}
 		}
@@ -103,7 +105,7 @@ namespace UE::UsdMediaSpatialAudioTranslator::Private
 		FString ErrorMessage;
 		if (!WaveInfo.ReadWaveInfo(FileBytes.GetData(), FileBytes.Num(), &ErrorMessage))
 		{
-			UE_LOG(LogUsd, Warning, TEXT("Failed to read audio info from file '%s': %s"), *FilePath, *ErrorMessage);
+			USD_LOG_WARNING(TEXT("Failed to read audio info from file '%s': %s"), *FilePath, *ErrorMessage);
 			return nullptr;
 		}
 
@@ -120,7 +122,7 @@ namespace UE::UsdMediaSpatialAudioTranslator::Private
 				WaveInfo = FWaveModInfo();
 				if (!WaveInfo.ReadWaveInfo(ConvertedRawWaveData.GetData(), ConvertedRawWaveData.Num(), &ErrorMessage))
 				{
-					UE_LOG(LogUsd, Warning, TEXT("Failed to convert audio data from '%s' into a supported wav format: %s"), *FilePath, *ErrorMessage);
+					USD_LOG_WARNING(TEXT("Failed to convert audio data from '%s' into a supported wav format: %s"), *FilePath, *ErrorMessage);
 					return nullptr;
 				}
 
@@ -134,7 +136,7 @@ namespace UE::UsdMediaSpatialAudioTranslator::Private
 
 		if (*WaveInfo.pBitsPerSample != 16)
 		{
-			UE_LOG(LogUsd, Warning, TEXT("Failed to parse audio as .wav data from file '%s' was not 16-bit!"), *FilePath);
+			USD_LOG_WARNING(TEXT("Failed to parse audio as .wav data from file '%s' was not 16-bit!"), *FilePath);
 			return nullptr;
 		}
 		check(*WaveInfo.pBitsPerSample == 16);
@@ -208,7 +210,11 @@ void FUsdMediaSpatialAudioTranslator::CreateAssets()
 	FString ResolvedAudioPath = UsdUtils::GetResolvedAssetPath(FilePathAttr, pxr::UsdTimeCode::Default());
 	if (!FPaths::FileExists(ResolvedAudioPath))
 	{
-		UE_LOG(LogUsd, Warning, TEXT("Failed to resolve audio file at path '%s' from prim '%s'"), *ResolvedAudioPath, *PrimPathString);
+		USD_LOG_USERWARNING(FText::Format(
+			LOCTEXT("MissingAudioFile", "Failed to resolve audio file at path '{0}' from prim '{1}'"),
+			FText::FromString(ResolvedAudioPath),
+			FText::FromString(PrimPathString)
+		));
 		return;
 	}
 
@@ -329,4 +335,6 @@ TSet<UE::FSdfPath> FUsdMediaSpatialAudioTranslator::CollectAuxiliaryPrims() cons
 	return {};
 }
 
+#undef LOCTEXT_NAMESPACE
+
 #endif	  // #if USE_USD_SDK
diff --git a/Engine/Plugins/Importers/USDImporter/Source/USDSchemas/Private/USDShadeMaterialTranslator.cpp b/Engine/Plugins/Importers/USDImporter/Source/USDSchemas/Private/USDShadeMaterialTranslator.cpp
--- a/Engine/Plugins/Importers/USDImporter/Source/USDSchemas/Private/USDShadeMaterialTranslator.cpp
+++ b/Engine/Plugins/Importers/USDImporter/Source/USDSchemas/Private/USDShadeMaterialTranslator.cpp
@@ -6,7 +6,7 @@
 #include "Objects/USDPrimLinkCache.h"
 #include "USDAssetUserData.h"
 #include "USDClassesModule.h"
-#include "USDLog.h"
+#include "USDErrorUtils.h"
 #include "USDMemory.h"
 #include "USDObjectUtils.h"
 #include "USDPrimConversion.h"
@@ -32,6 +32,8 @@
 #include "pxr/usd/usdShade/tokens.h"
 #include "USDIncludesEnd.h"
 
+#define LOCTEXT_NAMESPACE "USDShadeMaterialTranslator"
+
 namespace UE::UsdShadeTranslator::Private
 {
 	void RecursiveUpgradeMaterialsAndTexturesToVT(
@@ -53,7 +55,7 @@ namespace UE::UsdShadeTranslator::Private
 				continue;
 			}
 
-			UE_LOG(LogUsd, Log, TEXT("Upgrading texture '%s' to VT as it is used by a material that must be VT"), *Texture->GetName());
+			USD_LOG_INFO(TEXT("Upgrading texture '%s' to VT as it is used by a material that must be VT"), *Texture->GetName());
 			FPropertyChangedEvent PropertyChangeEvent(
 				UTexture::StaticClass()->FindPropertyByName(GET_MEMBER_NAME_CHECKED(UTexture, VirtualTextureStreaming))
 			);
@@ -114,9 +116,7 @@ namespace UE::UsdShadeTranslator::Private
 							Options = (FMaterialUpdateContext::EOptions::Type)(Options & ~FMaterialUpdateContext::EOptions::RecreateRenderStates);
 						}
 
-						UE_LOG(
-							LogUsd,
-							Log,
+						USD_LOG_INFO(
 							TEXT("Upgrading material instance '%s' to having a VT reference as texture '%s' requires it"),
 							*MaterialInstance->GetName(),
 							*Texture->GetName()
@@ -288,12 +288,6 @@ void FUsdShadeMaterialTranslator::CreateAssets()
 	{
 		if (!Context->UsdInfoCache->IsMaterialUsed(PrimPath))
 		{
-			UE_LOG(
-				LogUsd,
-				Verbose,
-				TEXT("Skipping creating assets for material prim '%s' as it is not currently bound by any prim."),
-				*PrimPath.GetString()
-			);
 			return;
 		}
 	}
@@ -309,11 +303,9 @@ void FUsdShadeMaterialTranslator::CreateAssets()
 		TOptional<FString> UnrealMaterial = UsdUtils::GetUnrealSurfaceOutput(ShadeMaterial.GetPrim());
 		if (UnrealMaterial.IsSet())
 		{
-			UE_LOG(
-				LogUsd,
-				Log,
+			USD_LOG_INFO(
 				TEXT(
-					"Skipping generation of assets for material prim '%s' as all prims that bind this material will use its referenced Unreal material '%s' instead."
+					"Skipping generation of assets for material prim '{0}' as all prims that bind this material will use its referenced Unreal material '{1}' instead."
 				),
 				*PrimPath.GetString(),
 				*UnrealMaterial.GetValue()
@@ -322,13 +314,11 @@ void FUsdShadeMaterialTranslator::CreateAssets()
 			UObject* Object = FSoftObjectPath(UnrealMaterial.GetValue()).TryLoad();
 			if (!Object)
 			{
-				UE_LOG(
-					LogUsd,
-					Warning,
-					TEXT("Failed to find the Unreal material '%s' referenced by material prim '%s'."),
-					*UnrealMaterial.GetValue(),
-					*PrimPath.GetString()
-				);
+				USD_LOG_USERWARNING(FText::Format(
+					LOCTEXT("MissingUnrealMaterial", "Failed to find the Unreal material '{0}' referenced by material prim '{1}'."),
+					FText::FromString(UnrealMaterial.GetValue()),
+					FText::FromString(PrimPath.GetString())
+				));
 			}
 
 			return;
@@ -405,7 +395,6 @@ void FUsdShadeMaterialTranslator::CreateAssets()
 			// Our VT material only has VT texture samplers, so *all* of its textures must be VT
 			if (VTTextures.Num() && NonVTTextures.Num())
 			{
-				UE_LOG(LogUsd, Log, TEXT("Upgrading textures used by material instance '%s' to VT as the material must be VT"), *MIC->GetName());
 				UE::UsdShadeTranslator::Private::UpgradeMaterialsAndTexturesToVT(NonVTTextures, Context);
 			}
 
@@ -523,7 +512,6 @@ void FUsdShadeMaterialTranslator::CreateAssets()
 			// Our VT material only has VT texture samplers, so *all* of its textures must be VT
 			if (VTTextures.Num() && NonVTTextures.Num())
 			{
-				UE_LOG(LogUsd, Log, TEXT("Upgrading textures used by material instance '%s' to VT as the material must be VT"), *MI->GetName());
 				UE::UsdShadeTranslator::Private::UpgradeMaterialsAndTexturesToVT(NonVTTextures, Context);
 			}
 
@@ -668,4 +656,6 @@ TSet<UE::FSdfPath> FUsdShadeMaterialTranslator::CollectAuxiliaryPrims() const
 	return Result;
 }
 
+#undef LOCTEXT_NAMESPACE
+
 #endif	  // #if USE_USD_SDK
diff --git a/Engine/Plugins/Importers/USDImporter/Source/USDSchemas/Private/USDSkelRootTranslator.cpp b/Engine/Plugins/Importers/USDImporter/Source/USDSchemas/Private/USDSkelRootTranslator.cpp
--- a/Engine/Plugins/Importers/USDImporter/Source/USDSchemas/Private/USDSkelRootTranslator.cpp
+++ b/Engine/Plugins/Importers/USDImporter/Source/USDSchemas/Private/USDSkelRootTranslator.cpp
@@ -451,13 +451,10 @@ namespace UsdSkelRootTranslatorImpl
 		InSkeletonCache.ComputeSkelBindings(InSkeletonRoot, &SkeletonBindings, pxr::UsdTraverseInstanceProxies());
 		if (SkeletonBindings.size() < 1)
 		{
-			FUsdLogManager::LogMessage(
-				EMessageSeverity::Warning,
-				FText::Format(
-					LOCTEXT("InvalidBinding", "SkelRoot {0} doesn't have any binding. No skinned mesh will be generated."),
-					FText::FromString(UsdToUnreal::ConvertPath(InSkeletonRoot.GetPath()))
-				)
-			);
+			USD_LOG_USERWARNING(FText::Format(
+				LOCTEXT("InvalidBinding", "SkelRoot {0} doesn't have any binding. No skinned mesh will be generated."),
+				FText::FromString(UsdToUnreal::ConvertPath(InSkeletonRoot.GetPath()))
+			));
 			return false;
 		}
 
diff --git a/Engine/Plugins/Importers/USDImporter/Source/USDSchemas/Private/USDSkelSkeletonTranslator.cpp b/Engine/Plugins/Importers/USDImporter/Source/USDSchemas/Private/USDSkelSkeletonTranslator.cpp
--- a/Engine/Plugins/Importers/USDImporter/Source/USDSchemas/Private/USDSkelSkeletonTranslator.cpp
+++ b/Engine/Plugins/Importers/USDImporter/Source/USDSchemas/Private/USDSkelSkeletonTranslator.cpp
@@ -12,11 +12,11 @@
 #include "USDClassesModule.h"
 #include "USDConversionUtils.h"
 #include "USDDrawModeComponent.h"
+#include "USDErrorUtils.h"
 #include "USDGeomMeshConversion.h"
 #include "USDGroomTranslatorUtils.h"
 #include "USDIntegrationUtils.h"
 #include "USDLayerUtils.h"
-#include "USDLog.h"
 #include "USDMemory.h"
 #include "USDObjectUtils.h"
 #include "USDPrimConversion.h"
@@ -139,9 +139,7 @@ namespace UsdSkelSkeletonTranslatorImpl
 			FSkeletalMeshLODInfo* LODInfo = SkeletalMesh->GetLODInfo(LODIndex);
 			if (!LODInfo)
 			{
-				UE_LOG(
-					LogUsd,
-					Error,
+				USD_LOG_ERROR(
 					TEXT("When processing materials for SkeletalMesh '%s', encountered no LOD info for LOD index %d!"),
 					*SkeletalMesh->GetName(),
 					LODIndex
@@ -169,9 +167,7 @@ namespace UsdSkelSkeletonTranslatorImpl
 				}
 				else
 				{
-					UE_LOG(
-						LogUsd,
-						Error,
+					USD_LOG_ERROR(
 						TEXT("Failed to resolve material '%s' for slot '%d' of LOD '%d' for mesh '%s'"),
 						*Slot.MaterialSource,
 						LODSlotIndex,
@@ -512,9 +508,8 @@ namespace UsdSkelSkeletonTranslatorImpl
 				}
 				else
 				{
-					UE_LOG(
-						LogUsd,
-						Log,
+					// TODO: Does this even make any sense? To be a "skinned prim" it *must* have the SkelBindingAPI, no?
+					USD_LOG_INFO(
 						TEXT("Ignoring skinned prim '%s' when generating Skeletal Mesh for Skeleton '%s' as the prim doesn't have the SkelBindingAPI"
 						),
 						*UsdToUnreal::ConvertPath(SkinningQuery.GetPrim().GetPrimPath()),
@@ -584,7 +579,7 @@ namespace UsdSkelSkeletonTranslatorImpl
 			bool bSuccess = UsdToUnreal::ConvertSkinnedMesh(SkinningQuery, SkeletonQuery, LODImportData, LODMaterialInfo, Options);
 			if (!bSuccess)
 			{
-				UE_LOG(LogUsd, Warning, TEXT("Failed to convert skinned mesh '%s'"), *UsdToUnreal::ConvertPath(LODMeshPrim.GetPrimPath()));
+				USD_LOG_WARNING(TEXT("Failed to convert skinned mesh '%s'"), *UsdToUnreal::ConvertPath(LODMeshPrim.GetPrimPath()));
 				return true;
 			}
 			bConvertedMeshData = true;
@@ -702,7 +697,7 @@ namespace UsdSkelSkeletonTranslatorImpl
 					}
 					else
 					{
-						UE_LOG(LogUsd, Error, TEXT("Failed to remap blend shape '%s's LOD index '%d'"), *BlendShape.Name, OldLODIndexUser);
+						USD_LOG_ERROR(TEXT("Failed to remap blend shape '%s's LOD index '%d'"), *BlendShape.Name, OldLODIndexUser);
 					}
 				}
 
@@ -897,9 +892,7 @@ namespace UsdSkelSkeletonTranslatorImpl
 				}
 				else
 				{
-					UE_LOG(
-						LogUsd,
-						Error,
+					USD_LOG_ERROR(
 						TEXT("Lost track of resolved material for slot '%d' of LOD '%d' for skeletal mesh '%s'"),
 						LODSlotIndex,
 						LODIndex,
@@ -1073,11 +1066,10 @@ namespace UsdSkelSkeletonTranslatorImpl
 
 				if (AnimBP->TargetSkeleton->GetReferenceSkeleton().GetRefBoneInfo() != Skeleton->GetReferenceSkeleton().GetRefBoneInfo())
 				{
-					UE_LOG(
-						LogUsd,
-						Warning,
-						TEXT("Forcing AnimBlueprint '%s's Skeleton '%s' to be compatible with the Skeleton generated for prim '%s', but they may be "
-							 "different!"),
+					USD_LOG_WARNING(
+						TEXT(
+							"Forcing AnimBlueprint '%s's Skeleton '%s' to be compatible with the Skeleton generated for prim '%s', but they may be different!"
+						),
 						*AnimBP->GetPathName(),
 						*AnimBP->TargetSkeleton->GetPathName(),
 						*PrimName
@@ -1203,15 +1195,15 @@ namespace UsdSkelSkeletonTranslatorImpl
 
 			if (LiveLinkNodes.Num() > 1 && !ExistingAnimBP)
 			{
-				UE_LOG(
-					LogUsd,
-					Warning,
-					TEXT("Found more than one LiveLinkPose blueprint node on AnimBlueprint '%s's graphs."
-						 "Note that all of those nodes will have their LiveLink SubjectName's updated to '%s', as described for prim '%s'!"),
-					*AnimBP->GetPathName(),
-					*SubjectName,
-					*UsdToUnreal::ConvertPath(SkeletonPrim.GetPrimPath())
-				);
+				USD_LOG_USERWARNING(FText::Format(
+					LOCTEXT(
+						"MoreThanOneLiveLinkPose",
+						"Found more than one LiveLinkPose blueprint node on AnimBlueprint '{0}'s graphs. Note that all of those nodes will have their LiveLink SubjectName's updated to '{1}', as described for prim '{2}'!"
+					),
+					FText::FromString(AnimBP->GetPathName()),
+					FText::FromString(SubjectName),
+					FText::FromString(UsdToUnreal::ConvertPath(SkeletonPrim.GetPrimPath()))
+				));
 			}
 
 			for (UAnimGraphNode_LiveLinkPose* Node : LiveLinkNodes)
@@ -1237,15 +1229,14 @@ namespace UsdSkelSkeletonTranslatorImpl
 				{
 					if (!ExistingAnimBP)
 					{
-						UE_LOG(
-							LogUsd,
-							Warning,
-							TEXT("Failed to update a LiveLinkPose node's 'Subject Name' to '%s' on AnimBlueprint '%s', "
-								 "because the pin is already connected to some other node. Disconnect it if you want it to be updated automatically."
+						USD_LOG_USERWARNING(FText::Format(
+							LOCTEXT(
+								"PinAlreadyConnected",
+								"Failed to update a LiveLinkPose node's 'Subject Name' to '{0}' on AnimBlueprint '{1}', because the pin is already connected to some other node. Disconnect it if you want it to be updated automatically."
 							),
-							*SubjectName,
-							*AnimBP->GetPathName()
-						);
+							FText::FromString(SubjectName),
+							FText::FromString(AnimBP->GetPathName())
+						));
 					}
 
 					continue;
@@ -1366,9 +1357,7 @@ namespace UsdSkelSkeletonTranslatorImpl
 
 			if (!bSuccess)
 			{
-				UE_LOG(
-					LogUsd,
-					Warning,
+				USD_LOG_WARNING(
 					TEXT("Failed to find skeleton binding info for skeleton at path '%s'. Is it within a SkelRoot prim?"),
 					*InPrimPath.GetString()
 				);
@@ -1535,9 +1524,7 @@ namespace UsdSkelSkeletonTranslatorImpl
 
 					if (!bSuccess)
 					{
-						UE_LOG(
-							LogUsd,
-							Warning,
+						USD_LOG_WARNING(
 							TEXT("Failed to create SkeletalMesh for prim '%s'"),
 							*UsdToUnreal::ConvertPath(ClosestParentSkelRoot.Get().GetPrim().GetPrimPath())
 						);
@@ -1545,9 +1532,7 @@ namespace UsdSkelSkeletonTranslatorImpl
 
 						if (bSkeletonIsNew)
 						{
-							UE_LOG(
-								LogUsd,
-								Warning,
+							USD_LOG_WARNING(
 								TEXT("Failed to create Skeleton '%s' for prim '%s'"),
 								*SkeletonName.ToString(),
 								*UsdToUnreal::ConvertPath(ClosestParentSkelRoot.Get().GetPrim().GetPrimPath())
@@ -1670,12 +1655,7 @@ namespace UsdSkelSkeletonTranslatorImpl
 								}
 								if (!bSuccess)
 								{
-									UE_LOG(
-										LogUsd,
-										Warning,
-										TEXT("Failed to create PhysicsAsset for skeletal mesh '%s'"),
-										*SkeletalMesh->GetPathName()
-									);
+									USD_LOG_WARNING(TEXT("Failed to create PhysicsAsset for skeletal mesh '%s'"), *SkeletalMesh->GetPathName());
 									UsdUnreal::TranslatorUtils::AbandonFailedAsset(
 										PhysicsAsset,
 										Context->UsdAssetCache.Get(),
@@ -1875,7 +1855,7 @@ namespace UsdSkelSkeletonTranslatorImpl
 					if (!bSuccess
 						|| (AnimSequence->GetDataModel()->GetNumBoneTracks() == 0 && AnimSequence->GetDataModel()->GetNumberOfFloatCurves() == 0))
 					{
-						UE_LOG(LogUsd, Warning, TEXT("Failed to create AnimSequence for prim '%s'"), *SkelAnimationPrimPath);
+						USD_LOG_WARNING(TEXT("Failed to create AnimSequence for prim '%s'"), *SkelAnimationPrimPath);
 						UsdUnreal::TranslatorUtils::AbandonFailedAsset(AnimSequence, Context->UsdAssetCache.Get(), Context->PrimLinkCache);
 						AnimSequence = nullptr;
 					}
@@ -1963,9 +1943,7 @@ USceneComponent* FUsdSkelSkeletonTranslator::CreateComponents()
 	pxr::UsdPrim ClosestParentSkelRoot = UsdUtils::GetClosestParentSkelRoot(SkeletonPrim);
 	if (!ClosestParentSkelRoot)
 	{
-		UE_LOG(
-			LogUsd,
-			Warning,
+		USD_LOG_WARNING(
 			TEXT("Ignoring skeleton '%s' when creating components as it is not contained within a SkelRoot prim scope"),
 			*PrimPath.GetString()
 		);
@@ -2026,9 +2004,7 @@ void FUsdSkelSkeletonTranslator::UpdateComponents(USceneComponent* SceneComponen
 	UE::FUsdPrim ClosestSkelRootPrim = UsdUtils::GetClosestParentSkelRoot(SkeletonPrim);
 	if (!ClosestSkelRootPrim)
 	{
-		UE_LOG(
-			LogUsd,
-			Warning,
+		USD_LOG_WARNING(
 			TEXT("Ignoring skeleton '%s' when updating components as it is not contained within a SkelRoot prim scope"),
 			*PrimPath.GetString()
 		);
@@ -2046,14 +2022,13 @@ void FUsdSkelSkeletonTranslator::UpdateComponents(USceneComponent* SceneComponen
 
 		// Commenting the usual deprecation macro so that we can find this with search and replace later
 		// UE_DEPRECATED(5.4, "schemas")
-		UE_LOG(
-			LogUsd,
-			Warning,
-			TEXT("Placing integration schemas (Live Link, Control Rig, Groom Binding) on SkelRoot prims (like '%s') has been deprecated on "
-				 "version 5.4 and will be unsupported in a future release. Please place your integration schemas directly on the Skeleton "
-				 "prims instead!"),
-			*ClosestSkelRootPrim.GetPrimPath().GetString()
-		);
+		USD_LOG_USERWARNING(FText::Format(
+			LOCTEXT(
+				"DeprecatedSchemas",
+				"Placing integration schemas (Live Link, Control Rig, Groom Binding) on SkelRoot prims (like '{0}') has been deprecated on version 5.4 and will be unsupported in a future release. Please place your integration schemas directly on the Skeleton prims instead!"
+			),
+			FText::FromString(ClosestSkelRootPrim.GetPrimPath().GetString())
+		));
 	}
 	const bool bPrimHasLiveLinkSchema = static_cast<bool>(PrimWithLiveLinkSchema);
 
@@ -2221,14 +2196,13 @@ void FUsdSkelSkeletonTranslator::UpdateComponents(USceneComponent* SceneComponen
 		{
 			// Commenting the usual deprecation macro so that we can find this with search and replace later
 			// UE_DEPRECATED(5.4, "schemas")
-			UE_LOG(
-				LogUsd,
-				Warning,
-				TEXT("Placing integration schemas (Live Link, Control Rig, Groom Binding) on SkelRoot prims (like '%s') has been deprecated on "
-					 "version 5.4 and will be unsupported in a future release. Please place your integration schemas directly on the Skeleton "
-					 "prims instead!"),
-				*ClosestSkelRootPrim.GetPrimPath().GetString()
-			);
+			USD_LOG_USERWARNING(FText::Format(
+				LOCTEXT(
+					"DeprecatedSchemas",
+					"Placing integration schemas (Live Link, Control Rig, Groom Binding) on SkelRoot prims (like '{0}') has been deprecated on version 5.4 and will be unsupported in a future release. Please place your integration schemas directly on the Skeleton prims instead!"
+				),
+				FText::FromString(ClosestSkelRootPrim.GetPrimPath().GetString())
+			));
 			UsdGroomTranslatorUtils::SetGroomFromPrim(ClosestSkelRootPrim, *Context->PrimLinkCache, SceneComponent);
 		}
 	}
diff --git a/Engine/Plugins/Importers/USDImporter/Source/USDSchemas/Private/USDVolVolumeTranslator.cpp b/Engine/Plugins/Importers/USDImporter/Source/USDSchemas/Private/USDVolVolumeTranslator.cpp
--- a/Engine/Plugins/Importers/USDImporter/Source/USDSchemas/Private/USDVolVolumeTranslator.cpp
+++ b/Engine/Plugins/Importers/USDImporter/Source/USDSchemas/Private/USDVolVolumeTranslator.cpp
@@ -7,9 +7,9 @@
 #include "USDAssetUserData.h"
 #include "USDConversionUtils.h"
 #include "USDDrawModeComponent.h"
+#include "USDErrorUtils.h"
 #include "USDIntegrationUtils.h"
 #include "USDLayerUtils.h"
-#include "USDLog.h"
 #include "USDMemory.h"
 #include "USDObjectUtils.h"
 #include "USDPrimConversion.h"
@@ -46,6 +46,8 @@
 #include <pxr/usd/usdVol/volume.h>
 #include "USDIncludesEnd.h"
 
+#define LOCTEXT_NAMESPACE "USDVolVolumeTranslator"
+
 namespace UE::UsdVolVolumeTranslator::Private
 {
 #if WITH_EDITOR
@@ -147,15 +149,16 @@ namespace UE::UsdVolVolumeTranslator::Private
 						check(ComponentIndex < SparseVolumeTextureChannelCount);
 						if (ComponentIndex < 0 || static_cast<uint32>(ComponentIndex) >= FoundGridInfo->NumComponents)
 						{
-							UE_LOG(
-								LogUsd,
-								Warning,
-								TEXT("Invalid component name '%s' for grid '%s' in VDB file '%s', as that particular grid only has %u components."),
-								*DesiredComponent,
-								*GridName,
-								*ParsedTexture.SourceVDBFilePath,
+							USD_LOG_USERWARNING(FText::Format(
+								LOCTEXT(
+									"InvalidComponentName",
+									"Invalid component name '{0}' for grid '{1}' in VDB file '{2}', as that particular grid only has {3} components."
+								),
+								FText::FromString(DesiredComponent),
+								FText::FromString(GridName),
+								FText::FromString(ParsedTexture.SourceVDBFilePath),
 								FoundGridInfo->NumComponents
-							);
+							));
 							ComponentIndex = INDEX_NONE;
 						}
 					}
@@ -169,17 +172,16 @@ namespace UE::UsdVolVolumeTranslator::Private
 							AvailableComponentNames = TEXT("'") + FString::Join(ComponentNames, TEXT("', '")) + TEXT("'");
 						}
 
-						UE_LOG(
-							LogUsd,
-							Warning,
-							TEXT(
-								"Desired component name '%s' for grid '%s' in VDB file '%s' is not a valid component name. Avaliable component names: %s"
+						USD_LOG_USERWARNING(FText::Format(
+							LOCTEXT(
+								"InvalidComponentName",
+								"Desired component name '{0}' for grid '{1}' in VDB file '{2}' is not a valid component name. Avaliable component names: {3}"
 							),
-							*DesiredComponent,
-							*GridName,
-							*ParsedTexture.SourceVDBFilePath,
-							*AvailableComponentNames
-						);
+							FText::FromString(DesiredComponent),
+							FText::FromString(GridName),
+							FText::FromString(ParsedTexture.SourceVDBFilePath),
+							FText::FromString(AvailableComponentNames)
+						));
 					}
 
 					// Get and validate desired channel (e.g. whether this mapping means to put something on 'AttributesA.R' or 'AttributesB.A', etc.)
@@ -194,17 +196,16 @@ namespace UE::UsdVolVolumeTranslator::Private
 
 						if (bIsSet && UsedChannels[ChannelIndex] != GridAndComponent)
 						{
-							UE_LOG(
-								LogUsd,
-								Warning,
-								TEXT(
-									"Cannot use attribute channel '%s' for grid '%s' in VDB file '%s', as the channel is already being used for the grid and component '%s'"
+							USD_LOG_USERWARNING(FText::Format(
+								LOCTEXT(
+									"AttributeChannelAlreadyUsed",
+									"Cannot use attribute channel '{0}' for grid '{1}' in VDB file '{2}', as the channel is already being used for the grid and component '{3}'"
 								),
-								*DesiredChannel,
-								*GridName,
-								*ParsedTexture.SourceVDBFilePath,
-								*UsedChannels[ChannelIndex]
-							);
+								FText::FromString(DesiredChannel),
+								FText::FromString(GridName),
+								FText::FromString(ParsedTexture.SourceVDBFilePath),
+								FText::FromString(UsedChannels[ChannelIndex])
+							));
 							ChannelIndex = INDEX_NONE;
 						}
 						else
@@ -223,17 +224,16 @@ namespace UE::UsdVolVolumeTranslator::Private
 							AvailableChannelNames = TEXT("'") + FString::Join(ChannelNames, TEXT("', '")) + TEXT("'");
 						}
 
-						UE_LOG(
-							LogUsd,
-							Warning,
-							TEXT(
-								"Desired attribute channel '%s' for grid '%s' in VDB file '%s' is not a valid channel name. Avaliable channel names: %s"
+						USD_LOG_USERWARNING(FText::Format(
+							LOCTEXT(
+								"InvalidAttributeChannel",
+								"Desired attribute channel '{0}' for grid '{1}' in VDB file '{2}' is not a valid channel name. Avaliable channel names: {3}"
 							),
-							*DesiredChannel,
-							*GridName,
-							*ParsedTexture.SourceVDBFilePath,
-							*AvailableChannelNames
-						);
+							FText::FromString(DesiredChannel),
+							FText::FromString(GridName),
+							FText::FromString(ParsedTexture.SourceVDBFilePath),
+							FText::FromString(AvailableChannelNames)
+						));
 					}
 
 					// Finally actually assign the desired grid/component mapping
@@ -265,14 +265,12 @@ namespace UE::UsdVolVolumeTranslator::Private
 					}
 				}
 
-				UE_LOG(
-					LogUsd,
-					Warning,
-					TEXT("Faild to find grid named '%s' inside VDB file '%s'. Avaliable grid names: %s"),
-					*GridName,
-					*ParsedTexture.SourceVDBFilePath,
-					*AvailableGrids
-				);
+				USD_LOG_USERWARNING(FText::Format(
+					LOCTEXT("InvalidGridName", "Faild to find grid named '{0}' inside VDB file '{1}'. Avaliable grid names: {2}"),
+					FText::FromString(GridName),
+					FText::FromString(ParsedTexture.SourceVDBFilePath),
+					FText::FromString(AvailableGrids)
+				));
 			}
 		}
 
@@ -437,16 +435,15 @@ namespace UE::UsdVolVolumeTranslator::Private
 				{
 					if (JobAtParam->SparseVolumeTexture != FoundParsedTexture->SparseVolumeTexture)
 					{
-						UE_LOG(
-							LogUsd,
-							Warning,
-							TEXT(
-								"Trying to assign different Sparse Volume Textures to the same material parameter '%s' on the material instantiated for Volume prim '%s' and field name '%s'! Only a single texture can be assigned to a material parameter at a time."
+						USD_LOG_USERWARNING(FText::Format(
+							LOCTEXT(
+								"MultipleSVTsOnSameParameter",
+								"Trying to assign different Sparse Volume Textures to the same material parameter '{0}' on the material instantiated for Volume prim '{1}' and field name '{2}'! Only a single texture can be assigned to a material parameter at a time."
 							),
-							*MaterialParameter,
-							*UsdToUnreal::ConvertPath(VolumePrim.GetPrimPath()),
-							*FieldName
-						);
+							FText::FromString(MaterialParameter),
+							FText::FromString(UsdToUnreal::ConvertPath(VolumePrim.GetPrimPath())),
+							FText::FromString(FieldName)
+						));
 					}
 				}
 				else
@@ -465,14 +462,15 @@ namespace UE::UsdVolVolumeTranslator::Private
 		{
 			if (!SparseVolumeTextureParameterNames.Contains(AssignmentPair.Key))
 			{
-				UE_LOG(
-					LogUsd,
-					Warning,
-					TEXT("Failed to assign Sparse Volume Texture '%s' to material '%s' as the desired material parameter '%s' doesn't exist on it"),
-					*AssignmentPair.Value->SparseVolumeTexture->GetPathName(),
-					*ReferenceMaterial->GetPathName(),
-					*AssignmentPair.Key
-				);
+				USD_LOG_USERWARNING(FText::Format(
+					LOCTEXT(
+						"MissingMaterialParameter",
+						"Failed to assign Sparse Volume Texture '{0}' to material '{1}' as the desired material parameter '{2}' doesn't exist on it"
+					),
+					FText::FromString(AssignmentPair.Value->SparseVolumeTexture->GetPathName()),
+					FText::FromString(ReferenceMaterial->GetPathName()),
+					FText::FromString(AssignmentPair.Key)
+				));
 			}
 		}
 
@@ -511,16 +509,15 @@ namespace UE::UsdVolVolumeTranslator::Private
 			// If there aren't enough parameters, let the user know
 			if ((SparseVolumeTextureParameterNames.Num() < FilePathHashToTextureInfo.Num()) && FilePathHashToTextureInfo.Num() > 0)
 			{
-				UE_LOG(
-					LogUsd,
-					Warning,
-					TEXT(
-						"Material '%s' used for prim '%s' doesn't have enough Sparse Volume Texture params to fit all of its %d parsed textures! Some may be left unassigned."
+				USD_LOG_USERWARNING(FText::Format(
+					LOCTEXT(
+						"PossibleUnassignedSVTs",
+						"Material '{0}' used for prim '{1}' doesn't have enough Sparse Volume Texture params to fit all of its {2} parsed textures! Some may be left unassigned."
 					),
-					*ReferenceMaterial->GetPathName(),
-					*UsdToUnreal::ConvertPath(VolumePrim.GetPrimPath()),
+					FText::FromString(ReferenceMaterial->GetPathName()),
+					FText::FromString(UsdToUnreal::ConvertPath(VolumePrim.GetPrimPath())),
 					FilePathHashToTextureInfo.Num()
-				);
+				));
 			}
 
 			TArray<FString> SortedFieldNames;
@@ -621,13 +618,11 @@ void FUsdVolVolumeTranslator::CreateAssets()
 
 		if (!FPaths::FileExists(VDBFilePath))
 		{
-			UE_LOG(
-				LogUsd,
-				Warning,
-				TEXT("Failed to find a VDB file at path '%s' when parsing Volume prim '%s'"),
-				*VDBFilePath,
-				*VolumePrimPathString
-			);
+			USD_LOG_USERWARNING(FText::Format(
+				LOCTEXT("MissingVDBFile", "Failed to find a VDB file at path '{0}' when parsing Volume prim '{1}'"),
+				FText::FromString(VDBFilePath),
+				FText::FromString(VolumePrimPathString)
+			));
 			continue;
 		}
 
@@ -697,7 +692,7 @@ void FUsdVolVolumeTranslator::CreateAssets()
 		);
 		if (!SparseVolumeTexture)
 		{
-			UE_LOG(LogUsd, Error, TEXT("Failed to generate Sparse Volume Texture from OpenVDB file '%s'"), *VDBFilePath);
+			USD_LOG_WARNING(TEXT("Failed to generate Sparse Volume Texture from OpenVDB file '%s'"), *VDBFilePath);
 			return;
 		}
 
@@ -809,14 +804,14 @@ void FUsdVolVolumeTranslator::CreateAssets()
 		// Warn in case the used material can't be used for SVTs
 		if (Material && Material->MaterialDomain != EMaterialDomain::MD_Volume)
 		{
-			UE_LOG(
-				LogUsd,
-				Warning,
-				TEXT("The material '%s' used for the Volume prim '%s' may not be capable of using Sparse Volume Textures as it does not have the "
-					 "Volume material domain."),
-				*ReferenceMaterial->GetPathName(),
-				*VolumePrimPathString
-			);
+			USD_LOG_USERWARNING(FText::Format(
+				LOCTEXT(
+					"WrongMaterialDomain",
+					"The material '{0}' used for the Volume prim '{1}' may not be capable of using Sparse Volume Textures as it does not have the Volume material domain."
+				),
+				FText::FromString(ReferenceMaterial->GetPathName()),
+				FText::FromString(VolumePrimPathString)
+			));
 		}
 	}
 	if (!Material)
@@ -1098,4 +1093,6 @@ TSet<UE::FSdfPath> FUsdVolVolumeTranslator::CollectAuxiliaryPrims() const
 	return Result;
 }
 
+#undef LOCTEXT_NAMESPACE
+
 #endif	  // #if USE_USD_SDK
diff --git a/Engine/Plugins/Importers/USDImporter/Source/USDStage/Private/USDLevelSequenceHelper.cpp b/Engine/Plugins/Importers/USDImporter/Source/USDStage/Private/USDLevelSequenceHelper.cpp
--- a/Engine/Plugins/Importers/USDImporter/Source/USDStage/Private/USDLevelSequenceHelper.cpp
+++ b/Engine/Plugins/Importers/USDImporter/Source/USDStage/Private/USDLevelSequenceHelper.cpp
@@ -7,10 +7,10 @@
 #include "USDAttributeUtils.h"
 #include "USDConversionUtils.h"
 #include "USDDrawModeComponent.h"
+#include "USDErrorUtils.h"
 #include "USDIntegrationUtils.h"
 #include "USDLayerUtils.h"
 #include "USDListener.h"
-#include "USDLog.h"
 #include "USDObjectUtils.h"
 #include "USDPrimConversion.h"
 #include "USDPrimTwin.h"
@@ -532,7 +532,7 @@ cleanup:
 			FText::FromString(OwnerActor->GetActorLabel())
 		);
 
-		UE_LOG(LogUsd, Warning, TEXT("%s"), *SubText.ToString().Replace(TEXT("\n\n"), TEXT(" ")));
+		USD_LOG_USERWARNING(FText::FromString(SubText.ToString().Replace(TEXT("\n\n"), TEXT(" "))));
 
 		const UUsdProjectSettings* Settings = GetDefault<UUsdProjectSettings>();
 		if (Settings && Settings->bShowTransformTrackOnCameraComponentWarning)
@@ -602,7 +602,7 @@ cleanup:
 			FText::FromName(PropertyName)
 		);
 
-		UE_LOG(LogUsd, Warning, TEXT("%s"), *SubText.ToString().Replace(TEXT("\n\n"), TEXT(" ")));
+		USD_LOG_USERWARNING(FText::FromString(SubText.ToString().Replace(TEXT("\n\n"), TEXT(" "))));
 
 		static TWeakPtr<SNotificationItem> Notification;
 
@@ -750,7 +750,7 @@ cleanup:
 			// Only show one at a time
 			if (!Notification.IsValid())
 			{
-				UE_LOG(LogUsd, Warning, TEXT("%s"), *SubText.ToString().Replace(TEXT("\n\n"), TEXT(" ")));
+				USD_LOG_USERWARNING(FText::FromString(SubText.ToString().Replace(TEXT("\n\n"), TEXT(" "))));
 				Notification = FSlateNotificationManager::Get().AddNotification(Toast);
 			}
 
@@ -1178,8 +1178,6 @@ void FUsdLevelSequenceHelperImpl::CreateLocalLayersSequences()
 	UE::FSdfLayer RootLayer = UsdStage.GetRootLayer();
 	const FLayerTimeInfo& RootLayerInfo = FindOrAddLayerTimeInfo(RootLayer);
 
-	UE_LOG(LogUsd, Verbose, TEXT("CreateLayerSequences: Initializing level sequence for '%s'"), *RootLayerInfo.Identifier);
-
 	// Create main level sequence for root layer
 	MainLevelSequence = FindOrAddSequenceForLayer(RootLayer, RootLayer.GetIdentifier(), RootLayer.GetDisplayName());
 	if (!MainLevelSequence)
@@ -1495,8 +1493,6 @@ ULevelSequence* FUsdLevelSequenceHelperImpl::FindOrAddSequenceForLayer(
 		const FLayerTimeInfo LayerTimeInfo = FindOrAddLayerTimeInfo(Layer);
 
 		UpdateMovieSceneTimeRanges(*MovieScene, LayerTimeInfo);
-
-		UE_LOG(LogUsd, Verbose, TEXT("Created Sequence for identifier: '%s'"), *SequenceIdentifier);
 	}
 
 	return Sequence;
@@ -1688,15 +1684,6 @@ void FUsdLevelSequenceHelperImpl::CreateSubSequenceSection(ULevelSequence& Seque
 			SubSectionRange.Size<FFrameNumber>().Value,
 			RowIndex
 		);
-
-		UE_LOG(
-			LogUsd,
-			Verbose,
-			TEXT("Adding subsection '%s' to sequence '%s'. StartFrame: '%d'"),
-			*SubSection->GetName(),
-			*Sequence.GetName(),
-			StartFrame.Value
-		);
 	}
 
 	const double TimeCodesPerSecondDifference = TimeCodesPerSecond / SubLayer.GetTimeCodesPerSecond();
@@ -1871,7 +1858,7 @@ void FUsdLevelSequenceHelperImpl::UpdateSubSectionTimeRanges(bool bShowResizedSe
 				"In USD, start and endTimeCodes from any layer but the root layer are mostly ignored, while in UE a subsequence section range does clip the underlying animation.\n\nIn order to reconcile the behavior, subsequence sections on any of the StageActor's generated LevelSequences will automatically snap to the stage's playback range instead, which prevents any unwanted animation clipping."
 			);
 
-			UE_LOG(LogUsd, Warning, TEXT("%s"), *SubText.ToString().Replace(TEXT("\n\n"), TEXT(" ")));
+			USD_LOG_USERWARNING(FText::FromString(SubText.ToString().Replace(TEXT("\n\n"), TEXT(" "))));
 
 			static TWeakPtr<SNotificationItem> Notification;
 
@@ -3000,9 +2987,7 @@ void FUsdLevelSequenceHelperImpl::AddVolumeTracks(const UUsdPrimTwin& PrimTwin,
 	}
 	else
 	{
-		UE_LOG(
-			LogUsd,
-			Warning,
+		USD_LOG_WARNING(
 			TEXT(
 				"Ignoring AssetUserData TimeSamplePathIndices when generating Sequencer tracks for Prim '%s' because it has %d entries, while it should have the same number of entries as TimeSamplePathTimeCodes (%d)"
 			),
@@ -3936,8 +3921,6 @@ void FUsdLevelSequenceHelperImpl::UpdateUsdLayerOffsetFromSection(const UMovieSc
 		return;
 	}
 
-	UE_LOG(LogUsd, Verbose, TEXT("Updating LevelSequence '%s' for sublayer '%s'"), *Sequence->GetName(), **SubLayerIdentifier);
-
 	const double TimeCodesPerSecond = GetTimeCodesPerSecond();
 	const double SubStartTimeCode = SubLayerTimeInfo->StartTimeCode.Get(0.0);
 	const double SubEndTimeCode = SubLayerTimeInfo->EndTimeCode.Get(0.0);
@@ -3980,7 +3963,7 @@ void FUsdLevelSequenceHelperImpl::UpdateUsdLayerOffsetFromSection(const UMovieSc
 		UE::FSdfLayer Layer = UE::FSdfLayer::FindOrOpen(*LayerTimeInfo->Identifier);
 		if (!Layer)
 		{
-			UE_LOG(LogUsd, Warning, TEXT("Failed to update sublayer '%s'"), *LayerTimeInfo->Identifier);
+			USD_LOG_WARNING(TEXT("Failed to update sublayer '%s'"), *LayerTimeInfo->Identifier);
 			return;
 		}
 
@@ -4012,8 +3995,6 @@ void FUsdLevelSequenceHelperImpl::UpdateUsdLayerOffsetFromSection(const UMovieSc
 			UsdUtils::SetRefOrPayloadLayerOffset(UsdStage.GetPrimAtPath(UE::FSdfPath(*PrimPath)), NewLayerOffset);
 		}
 	}
-
-	UE_LOG(LogUsd, Verbose, TEXT("\tNew OffsetScale: %f, %f"), NewLayerOffset.Offset, NewLayerOffset.Scale);
 }
 
 void FUsdLevelSequenceHelperImpl::UpdateMovieSceneTimeRanges(UMovieScene& MovieScene, const FLayerTimeInfo& LayerTimeInfo, bool bUpdateViewRanges)
@@ -5110,9 +5091,7 @@ void FUsdLevelSequenceHelperImpl::HandleTrackChange(UMovieSceneTrack& Track, boo
 				const TArray<UMovieSceneSection*>& Sections = AudioTrack->GetAudioSections();
 				if (Sections.Num() > 1)
 				{
-					UE_LOG(
-						LogUsd,
-						Warning,
+					USD_LOG_WARNING(
 						TEXT(
 							"The audio track '%s' has %d sections, but only the first audio section of an audio track can be written out to USD for now"
 						),
@@ -5150,15 +5129,6 @@ FUsdLevelSequenceHelperImpl::FLayerTimeInfo& FUsdLevelSequenceHelperImpl::FindOr
 	FLayerTimeInfo LayerTimeInfo;
 	UpdateLayerTimeInfoFromLayer(LayerTimeInfo, Layer);
 
-	UE_LOG(
-		LogUsd,
-		Verbose,
-		TEXT("Creating layer time info for layer '%s'. Original timecodes: ['%s', '%s']"),
-		*LayerTimeInfo.Identifier,
-		LayerTimeInfo.StartTimeCode.IsSet() ? *LexToString(LayerTimeInfo.StartTimeCode.GetValue()) : TEXT("null"),
-		LayerTimeInfo.EndTimeCode.IsSet() ? *LexToString(LayerTimeInfo.EndTimeCode.GetValue()) : TEXT("null")
-	);
-
 	return LayerTimeInfosByLayerIdentifier.Add(Layer.GetIdentifier(), LayerTimeInfo);
 }
 
@@ -5183,9 +5153,7 @@ void FUsdLevelSequenceHelperImpl::UpdateLayerTimeInfoFromLayer(FLayerTimeInfo& L
 	if (LayerTimeInfo.StartTimeCode.IsSet() && LayerTimeInfo.EndTimeCode.IsSet()
 		&& LayerTimeInfo.EndTimeCode.GetValue() < LayerTimeInfo.StartTimeCode.GetValue())
 	{
-		UE_LOG(
-			LogUsd,
-			Warning,
+		USD_LOG_WARNING(
 			TEXT("Sublayer '%s' has end time code (%f) before start time code (%f)! These values will be automatically swapped"),
 			*Layer.GetIdentifier(),
 			LayerTimeInfo.EndTimeCode.GetValue(),
diff --git a/Engine/Plugins/Importers/USDImporter/Source/USDStage/Private/USDStageActor.cpp b/Engine/Plugins/Importers/USDImporter/Source/USDStage/Private/USDStageActor.cpp
--- a/Engine/Plugins/Importers/USDImporter/Source/USDStage/Private/USDStageActor.cpp
+++ b/Engine/Plugins/Importers/USDImporter/Source/USDStage/Private/USDStageActor.cpp
@@ -22,7 +22,6 @@
 #include "USDLayerUtils.h"
 #include "USDLightConversion.h"
 #include "USDListener.h"
-#include "USDLog.h"
 #include "USDMemory.h"
 #include "USDObjectUtils.h"
 #include "USDPrimConversion.h"
@@ -681,7 +680,7 @@ struct FUsdStageActorImpl
 			FText::FromString(Component->GetName())
 		);
 
-		UE_LOG(LogUsd, Warning, TEXT("%s"), *SubText.ToString().Replace(TEXT("\n\n"), TEXT(" ")));
+		USD_LOG_USERWARNING(FText::FromString(SubText.ToString().Replace(TEXT("\n\n"), TEXT(" "))));
 
 		const UUsdProjectSettings* Settings = GetDefault<UUsdProjectSettings>();
 		if (Settings && Settings->bShowTransformOnCameraComponentWarning)
@@ -1171,7 +1170,10 @@ void AUsdStageActor::SetIsolatedRootLayer(const FString& IsolatedStageRootLayer)
 		}
 		else
 		{
-			UE_LOG(LogUsd, Warning, TEXT("Failed to isolate layer '%s': File does not exist or is not a valid USD layer"), *IsolatedStageRootLayer);
+			USD_LOG_USERWARNING(FText::Format(
+				LOCTEXT("FailIsolateMissingFile", "Failed to isolate layer '{0}': File does not exist or is not a valid USD layer"),
+				FText::FromString(IsolatedStageRootLayer)
+			));
 		}
 	}
 }
@@ -1237,13 +1239,13 @@ void AUsdStageActor::IsolateLayer(const UE::FSdfLayer& Layer, bool bLoadUsdStage
 
 		if (!ValidLayers.Contains(Layer))
 		{
-			UE_LOG(
-				LogUsd,
-				Warning,
-				TEXT("Failed to isolate layer '%s' as it is not part of the currently opened USD Stage's local layer stack"),
-				*Layer.GetIdentifier(),
-				*UsdStage.GetRootLayer().GetIdentifier()
-			);
+			USD_LOG_USERWARNING(FText::Format(
+				LOCTEXT(
+					"FailIsolateNonLocal",
+					"Failed to isolate layer '{0}' as it is not part of the currently opened USD Stage's local layer stack"
+				),
+				FText::FromString(Layer.GetIdentifier())
+			));
 			return;
 		}
 
@@ -1811,8 +1813,6 @@ void AUsdStageActor::HandleAccumulatedNotices()
 			const UE::FSdfPath& PrimPath = PrimChangedInfo.Key;
 			const bool bIsResync = PrimChangedInfo.Value == EPrimUpdateType::Resync;
 
-			UE_LOG(LogUsd, Verbose, TEXT("Cleaning up assets for prim path '%s'"), *PrimPath.GetString());
-
 			const bool bForEntireSubtree = bIsResync;
 			bHasLoadedOrAbandonedAssets |= UnloadAssets(PrimPath, bForEntireSubtree);
 		}
@@ -1830,8 +1830,6 @@ void AUsdStageActor::HandleAccumulatedNotices()
 			const UE::FSdfPath PrimPath = PrimChangedInfo.Key;
 			const bool bIsResync = PrimChangedInfo.Value == EPrimUpdateType::Resync;
 
-			UE_LOG(LogUsd, Verbose, TEXT("Cleaning up actors and components for prim path '%s'"), *PrimPath.GetString());
-
 			if (bIsResync && PrimPath.IsAbsoluteRootOrPrimPath())
 			{
 				if (UUsdPrimTwin* UsdPrimTwin = GetRootPrimTwin()->Find(PrimPath.GetString()))
@@ -1881,8 +1879,6 @@ void AUsdStageActor::HandleAccumulatedNotices()
 			}
 			ensure(PrimToUpdate);
 
-			UE_LOG(LogUsd, Verbose, TEXT("Regenerating assets for prim path '%s'"), *PrimPath.GetString());
-
 			TSharedRef<FUsdSchemaTranslationContext> TranslationContext = FUsdStageActorImpl::CreateUsdSchemaTranslationContext(
 				this,
 				PrimPath.GetString()
@@ -1918,8 +1914,6 @@ void AUsdStageActor::HandleAccumulatedNotices()
 				continue;
 			}
 
-			UE_LOG(LogUsd, Verbose, TEXT("Regenerating components for prim path '%s'"), *PrimPath.GetString());
-
 			TSharedRef<FUsdSchemaTranslationContext> TranslationContext = FUsdStageActorImpl::CreateUsdSchemaTranslationContext(
 				this,
 				PrimPath.GetString()
@@ -2310,14 +2304,13 @@ UUsdPrimTwin* AUsdStageActor::ExpandPrim(
 			{
 				// Commenting the usual deprecation macro so that we can find this with search and replace later
 				// UE_DEPRECATED(5.4, "schemas")
-				UE_LOG(
-					LogUsd,
-					Warning,
-					TEXT("Placing integration schemas (Live Link, Control Rig, Groom Binding) on SkelRoot prims (like '%s') has been deprecated on "
-						 "version 5.4 and will be unsupported in a future release. Please place your integration schemas directly on the Skeleton "
-						 "prims instead!"),
-					*Prim.GetPrimPath().GetString()
-				);
+				USD_LOG_USERWARNING(FText::Format(
+					LOCTEXT(
+						"DeprecatedSchemas",
+						"Placing integration schemas (Live Link, Control Rig, Groom Binding) on SkelRoot prims (like '{0}') has been deprecated on version 5.4 and will be unsupported in a future release. Please place your integration schemas directly on the Skeleton prims instead!"
+					),
+					FText::FromString(Prim.GetPrimPath().GetString())
+				));
 				PrimWithSchema = ParentSkelRoot;
 			}
 		}
@@ -3036,7 +3029,7 @@ void AUsdStageActor::OpenUsdStage()
 		RequestDelayedTransactorReset();
 	}
 
-	UsdUtils::StartMonitoringErrors();
+	FScopedUsdMessageLog ScopedLog;
 
 	FString AbsPath;
 	if (!RootLayer.FilePath.StartsWith(UnrealIdentifiers::IdentifierPrefix) && FPaths::IsRelative(RootLayer.FilePath))
@@ -3080,11 +3073,6 @@ void AUsdStageActor::OpenUsdStage()
 	}
 
 	OnStageChanged.Broadcast();
-
-	UsdUtils::ShowErrorsAndStopMonitoring(FText::Format(
-		LOCTEXT("USDOpenError", "Encountered some errors opening USD file at path '{0}!\nCheck the Output Log for details."),
-		FText::FromString(RootLayer.FilePath)
-	));
 }
 
 void AUsdStageActor::CloseUsdStage(bool bUnloadIfNeeded)
@@ -3373,15 +3361,7 @@ void AUsdStageActor::LoadUsdStage(bool bOpenIfNeeded)
 	int ElapsedMin = int(ElapsedSeconds / 60.0);
 	ElapsedSeconds -= 60.0 * (double)ElapsedMin;
 
-	UE_LOG(
-		LogUsd,
-		Log,
-		TEXT("%s %s in [%d min %.3f s]"),
-		TEXT("Stage loaded"),
-		*FPaths::GetBaseFilename(RootLayer.FilePath),
-		ElapsedMin,
-		ElapsedSeconds
-	);
+	USD_LOG_INFO(TEXT("%s %s in [%d min %.3f s]"), TEXT("Stage loaded"), *FPaths::GetBaseFilename(RootLayer.FilePath), ElapsedMin, ElapsedSeconds);
 
 #if USE_USD_SDK
 	FUsdStageActorImpl::SendAnalytics(this, ElapsedSeconds, UsdUtils::GetUsdStageNumFrames(StageToLoad), FPaths::GetExtension(RootLayer.FilePath));
diff --git a/Engine/Plugins/Importers/USDImporter/Source/USDStage/Private/USDTransactor.cpp b/Engine/Plugins/Importers/USDImporter/Source/USDStage/Private/USDTransactor.cpp
--- a/Engine/Plugins/Importers/USDImporter/Source/USDStage/Private/USDTransactor.cpp
+++ b/Engine/Plugins/Importers/USDImporter/Source/USDStage/Private/USDTransactor.cpp
@@ -3,9 +3,9 @@
 #include "USDTransactor.h"
 
 #include "UnrealUSDWrapper.h"
+#include "USDErrorUtils.h"
 #include "USDLayerUtils.h"
 #include "USDListener.h"
-#include "USDLog.h"
 #include "USDMemory.h"
 #include "USDPrimConversion.h"
 #include "USDStageActor.h"
@@ -193,9 +193,7 @@ namespace UsdUtils
 						ConvertedChange.PrimAppliedSchemas.Add(SchemaName.ToString());
 					}
 
-					UE_LOG(
-						LogUsd,
-						Log,
+					USD_LOG_INFO(
 						TEXT("Recorded the %s of prim '%s' with TypeName '%s' and PrimAppliedSchemas [%s]"),
 						(ConvertedChange.Flags.bDidAddInertPrim || ConvertedChange.Flags.bDidAddNonInertPrim) ? TEXT("addition") : TEXT("removal"),
 						*ObjectPath,
@@ -210,7 +208,7 @@ namespace UsdUtils
 					// We expect not to find the prim if the change says it was just removed though
 					if (!ConvertedChange.Flags.bDidRemoveInertPrim && !ConvertedChange.Flags.bDidRemoveNonInertPrim)
 					{
-						UE_LOG(LogUsd, Warning, TEXT("Failed to find prim at path '%s' when serializing object changes in transactor"), *ObjectPath);
+						USD_LOG_WARNING(TEXT("Failed to find prim at path '%s' when serializing object changes in transactor"), *ObjectPath);
 					}
 					continue;
 				}
@@ -239,11 +237,10 @@ namespace UsdUtils
 							ConvertedChange.Flags.bDidRemovePropertyWithOnlyRequiredFields))
 					{
 						ConvertedAttributeChange.AttributeTypeName = UsdUtils::GetImpliedTypeName(FieldChange.OldValue);
-						UE_LOG(
-							LogUsd,
-							Log,
-							TEXT("Recording the removal of properties is not fully supported: Using underlying type '%s' for record of attribute "
-								 "'%s' of prim '%s', as we don't have access to the attribute's role"),
+						USD_LOG_INFO(
+							TEXT(
+								"Recording the removal of properties is not fully supported: Using underlying type '%s' for record of attribute '%s' of prim '%s', as we don't have access to the attribute's role"
+							),
 							*ConvertedAttributeChange.AttributeTypeName,
 							*PropertyName,
 							*ObjectPath
@@ -277,9 +274,7 @@ namespace UsdUtils
 						}
 						else
 						{
-							UE_LOG(
-								LogUsd,
-								Warning,
+							USD_LOG_WARNING(
 								TEXT("Failed to find attribute '%s' for prim at path '%s' when serializing object changes in transactor"),
 								*PropertyName,
 								*PrimPath.GetString()
@@ -306,7 +301,7 @@ namespace UsdUtils
 
 		if ((bAdd && Direction == EApplicationDirection::Forward) || (bRemove && Direction == EApplicationDirection::Reverse))
 		{
-			UE_LOG(LogUsd, Log, TEXT("Creating prim '%s' with typename '%s'"), *PrimPath.GetString(), *PrimChange.PrimTypeName);
+			USD_LOG_INFO(TEXT("Creating prim '%s' with typename '%s'"), *PrimPath.GetString(), *PrimChange.PrimTypeName);
 
 			UE::FUsdPrim Prim = Stage.DefinePrim(PrimPath, *PrimChange.PrimTypeName);
 			if (!Prim)
@@ -318,7 +313,7 @@ namespace UsdUtils
 		}
 		else if ((bAdd && Direction == EApplicationDirection::Reverse) || (bRemove && Direction == EApplicationDirection::Forward))
 		{
-			UE_LOG(LogUsd, Log, TEXT("Removing prim '%s' with typename '%s'"), *PrimPath.GetString(), *PrimChange.PrimTypeName);
+			USD_LOG_INFO(TEXT("Removing prim '%s' with typename '%s'"), *PrimPath.GetString(), *PrimChange.PrimTypeName);
 
 			UsdUtils::RemoveAllLocalPrimSpecs(Stage.GetPrimAtPath(PrimPath), Stage.GetEditTarget());
 			return true;
@@ -347,7 +342,7 @@ namespace UsdUtils
 
 			if (UE::FUsdPrim Prim = Stage.GetPrimAtPath(UsdCurrentPath))
 			{
-				UE_LOG(LogUsd, Log, TEXT("Renaming prim '%s' to '%s'"), *Prim.GetPrimPath().GetString(), *NewName);
+				USD_LOG_INFO(TEXT("Renaming prim '%s' to '%s'"), *Prim.GetPrimPath().GetString(), *NewName);
 
 				const bool bDidRename = UsdUtils::RenamePrim(Prim, *NewName);
 				if (bDidRename)
@@ -364,7 +359,7 @@ namespace UsdUtils
 				return true;
 			}
 
-			UE_LOG(LogUsd, Warning, TEXT("Failed to rename prim at path '%s' to name '%s'"), *CurrentPath, *NewName);
+			USD_LOG_WARNING(TEXT("Failed to rename prim at path '%s' to name '%s'"), *CurrentPath, *NewName);
 		}
 
 		return false;
@@ -395,9 +390,7 @@ namespace UsdUtils
 		// This corresponds to the token SdfFieldKeys->TimeSamples, and is extracted from UsdAttribute::Clear
 		Attribute.ClearMetadata(TEXT("timeSamples"));
 
-		UE_LOG(
-			LogUsd,
-			Log,
+		USD_LOG_INFO(
 			TEXT("Applying '%d' timeSamples for attribute '%s' of prim '%s'"),
 			AttributeChange.TimeSamples.Num(),
 			*PropertyName,
@@ -412,9 +405,7 @@ namespace UsdUtils
 			{
 				if (!Attribute.Set(Value, AttributeChange.TimeSamples[Index]))
 				{
-					UE_LOG(
-						LogUsd,
-						Warning,
+					USD_LOG_WARNING(
 						TEXT("Failed to apply value '%s' at timesample '%f' for attribute '%s' of prim '%s'"),
 						*UsdUtils::Stringify(Value),
 						AttributeChange.TimeSamples[Index],
@@ -426,9 +417,7 @@ namespace UsdUtils
 			}
 			else
 			{
-				UE_LOG(
-					LogUsd,
-					Warning,
+				USD_LOG_WARNING(
 					TEXT("Failed to convert value for timesample '%f' for attribute '%s' of prim '%s'"),
 					AttributeChange.TimeSamples[Index],
 					*PropertyName,
@@ -492,13 +481,7 @@ namespace UsdUtils
 				Relationship = Prim.CreateRelationship(*PropertyName);
 				if (!Relationship)
 				{
-					UE_LOG(
-						LogUsd,
-						Warning,
-						TEXT("Failed to create relationship '%s' with for prim '%s'"),
-						*PropertyName,
-						*Prim.GetPrimPath().GetString()
-					);
+					USD_LOG_WARNING(TEXT("Failed to create relationship '%s' with for prim '%s'"), *PropertyName, *Prim.GetPrimPath().GetString());
 
 					return false;
 				}
@@ -514,9 +497,7 @@ namespace UsdUtils
 					// We expect to fail to create an attribute if we have no typename here (e.g. undo remove property)
 					if (AttributeTypeName.IsEmpty())
 					{
-						UE_LOG(
-							LogUsd,
-							Warning,
+						USD_LOG_WARNING(
 							TEXT("Failed to create attribute '%s' with typename '%s' for prim '%s'"),
 							*PropertyName,
 							*AttributeTypeName,
@@ -534,9 +515,7 @@ namespace UsdUtils
 		UE::FVtValue WrapperValue;
 		if (!UnrealToUsd::ConvertValue(Value, WrapperValue))
 		{
-			UE_LOG(
-				LogUsd,
-				Warning,
+			USD_LOG_WARNING(
 				TEXT("Failed to convert VtValue back to USD when applying it to object '%s' field '%s'"),
 				*ObjectPath.GetString(),
 				*Field
@@ -544,9 +523,7 @@ namespace UsdUtils
 			return false;
 		}
 
-		UE_LOG(
-			LogUsd,
-			Log,
+		USD_LOG_INFO(
 			TEXT("%s object '%s' (typename '%s'), field '%s' with value '%s' at time '%s'"),
 			bCreated				 ? TEXT("Creating")
 			: WrapperValue.IsEmpty() ? bRemoveProperty ? TEXT("Removing") : TEXT("Clearing")
@@ -670,11 +647,11 @@ namespace UsdUtils
 		UE::FVtValue WrapperValue;
 		if (!UnrealToUsd::ConvertValue(Value, WrapperValue))
 		{
-			UE_LOG(LogUsd, Warning, TEXT("Failed to convert VtValue back to USD when applying it to stage metadata field '%s'"), *FieldName);
+			USD_LOG_WARNING(TEXT("Failed to convert VtValue back to USD when applying it to stage metadata field '%s'"), *FieldName);
 			return false;
 		}
 
-		UE_LOG(LogUsd, Log, TEXT("Setting stage metadata '%s', with value '%s'"), *FieldName, *UsdUtils::Stringify(WrapperValue));
+		USD_LOG_INFO(TEXT("Setting stage metadata '%s', with value '%s'"), *FieldName, *UsdUtils::Stringify(WrapperValue));
 
 		UE::FSdfLayer OldEditTarget = Stage.GetEditTarget();
 		Stage.SetEditTarget(Stage.GetRootLayer());	  // Stage metadata always needs to be set at the root layer
@@ -746,9 +723,7 @@ namespace UsdUtils
 				UE::FSdfLayer EditTarget = UsdUtils::FindLayerForIdentifier(*Edits.EditTargetIdentifier, Stage);
 				if (!EditTarget)
 				{
-					UE_LOG(
-						LogUsd,
-						Warning,
+					USD_LOG_WARNING(
 						TEXT("Ignoring application of recorded USD stage changes as the edit target with identifier '%s' cannot be found or opened"),
 						*Edits.EditTargetIdentifier
 					)
diff --git a/Engine/Plugins/Importers/USDImporter/Source/USDStageEditor/Private/SUSDLayersTreeView.cpp b/Engine/Plugins/Importers/USDImporter/Source/USDStageEditor/Private/SUSDLayersTreeView.cpp
--- a/Engine/Plugins/Importers/USDImporter/Source/USDStageEditor/Private/SUSDLayersTreeView.cpp
+++ b/Engine/Plugins/Importers/USDImporter/Source/USDStageEditor/Private/SUSDLayersTreeView.cpp
@@ -6,9 +6,9 @@
 #include "SUSDStageEditorStyle.h"
 #include "USDClassesModule.h"
 #include "USDDragDropOp.h"
+#include "USDErrorUtils.h"
 #include "USDLayersViewModel.h"
 #include "USDLayerUtils.h"
-#include "USDLog.h"
 #include "USDMemory.h"
 #include "USDProjectSettings.h"
 #include "UsdWrappers/SdfChangeBlock.h"
@@ -45,7 +45,7 @@ namespace UE::USDLayersTreeViewImpl::Private
 		UE::FSdfLayer OutputLayer = UE::FSdfLayer::CreateNew(*TargetPath);
 		if (!OutputLayer)
 		{
-			UE_LOG(LogUsd, Warning, TEXT("Failed to export USD layer to path '%s'"), *TargetPath);
+			USD_LOG_WARNING(TEXT("Failed to export USD layer to path '%s'"), *TargetPath);
 			return;
 		}
 
@@ -583,11 +583,13 @@ FReply SUsdLayersTreeView::OnRowDragDetected(const FGeometry& Geometry, const FP
 	}
 
 	const FSlateBrush* Icon = FAppStyle::GetBrush(TEXT("Layer.Icon16x"));
-	FText DefaultHoverText = DraggedItems.Num() == 1 ? FText::Format(
-								 LOCTEXT("DefaultTextSingle", "USD layer '{0}'"),
-								 FText::FromString(StaticCastSharedRef<FUsdLayerViewModel>(*DraggedItems.CreateConstIterator())->LayerIdentifier)
-							 )
-													 : FText::Format(LOCTEXT("DefaultTextMultiple", "{0} USD layers"), DraggedItems.Num());
+	FText DefaultHoverText = DraggedItems.Num() == 1
+								 ? FText::Format(
+									   LOCTEXT("DefaultTextSingle", "USD layer '{0}'"),
+									   FText::FromString(StaticCastSharedRef<FUsdLayerViewModel>(*DraggedItems.CreateConstIterator())->LayerIdentifier
+									   )
+								   )
+								 : FText::Format(LOCTEXT("DefaultTextMultiple", "{0} USD layers"), DraggedItems.Num());
 
 	TSharedRef<FUsdDragDropOp> Op = MakeShared<FUsdDragDropOp>();
 	Op->OpType = EUsdDragDropOpType::Layers;
@@ -1054,9 +1056,7 @@ FReply SUsdLayersTreeView::OnRowAcceptDrop(const FDragDropEvent& Event, EItemDro
 
 				if (!bInserted)
 				{
-					UE_LOG(
-						LogUsd,
-						Warning,
+					USD_LOG_WARNING(
 						TEXT("Failed to insert layer '%s' as a sublayer of '%s' at index '%d'"),
 						*LayerPin.GetIdentifier(),
 						*ParentLayer.GetIdentifier(),
@@ -1276,10 +1276,9 @@ TSharedPtr<SWidget> SUsdLayersTreeView::ConstructLayerContextMenu()
 						}
 					}
 
-					return bCanMute
-							   ? (bAnyIsMuted ? LOCTEXT("UnmuteLayerRightClick_ToolTip", "Unmute the selected muted layers")
-											  : LOCTEXT("MuteLayerRightClick_ToolTip", "Mute all selected layers"))
-							   : LOCTEXT("CantMuteLayerRightClick_ToolTip", "Cannot mute root layers or edit targets");
+					return bCanMute ? (bAnyIsMuted ? LOCTEXT("UnmuteLayerRightClick_ToolTip", "Unmute the selected muted layers")
+												   : LOCTEXT("MuteLayerRightClick_ToolTip", "Mute all selected layers"))
+									: LOCTEXT("CantMuteLayerRightClick_ToolTip", "Cannot mute root layers or edit targets");
 				}
 			)),
 			FSlateIcon(),
diff --git a/Engine/Plugins/Importers/USDImporter/Source/USDStageEditor/Private/SUSDStage.cpp b/Engine/Plugins/Importers/USDImporter/Source/USDStageEditor/Private/SUSDStage.cpp
--- a/Engine/Plugins/Importers/USDImporter/Source/USDStageEditor/Private/SUSDStage.cpp
+++ b/Engine/Plugins/Importers/USDImporter/Source/USDStageEditor/Private/SUSDStage.cpp
@@ -12,7 +12,6 @@
 #include "USDConversionUtils.h"
 #include "USDErrorUtils.h"
 #include "USDLayerUtils.h"
-#include "USDLog.h"
 #include "USDMaterialUtils.h"
 #include "USDProjectSettings.h"
 #include "USDStageActor.h"
@@ -2736,7 +2735,10 @@ void SUsdStage::FileExportFlattenedStage(const FString& OutputLayer)
 	const bool bResult = UsdStage.Export(*OutputLayerCopy);
 	if (!bResult)
 	{
-		UE_LOG(LogUsd, Warning, TEXT("Failed to export flattened USD Stage to path '%s'!"), *OutputLayerCopy);
+		USD_LOG_USERERROR(FText::Format(
+			LOCTEXT("FailedToExportFlattened", "Failed to export flattened USD Stage to path '{0}'!"),
+			FText::FromString(OutputLayerCopy)
+		));
 	}
 
 	// Send analytics
@@ -2784,9 +2786,10 @@ void SUsdStage::FileExportFlattenedLayerStack(const FString& OutputLayer)
 
 	double StartTime = FPlatformTime::Cycles64();
 
-	UsdUtils::StartMonitoringErrors();
+	FScopedUsdMessageLog ScopedLog;
+
 	UE::FSdfLayer FlattenedLayer = UsdUtils::FlattenLayerStack(UsdStage);
-	if (UsdUtils::ShowErrorsAndStopMonitoring())
+	if (FUsdLogManager::HasAccumulatedErrors())
 	{
 		return;
 	}
@@ -2794,7 +2797,10 @@ void SUsdStage::FileExportFlattenedLayerStack(const FString& OutputLayer)
 	const bool bResult = FlattenedLayer.Export(*OutputLayerCopy);
 	if (!bResult)
 	{
-		UE_LOG(LogUsd, Warning, TEXT("Failed to export flattened USD Stage to path '%s'!"), *OutputLayerCopy);
+		USD_LOG_USERERROR(FText::Format(
+			LOCTEXT("FailedToExportFlattenedLayerStack", "Failed to export flattened layer stack to path '{0}'!"),
+			FText::FromString(OutputLayerCopy)
+		));
 	}
 
 	// Send analytics
diff --git a/Engine/Plugins/Importers/USDImporter/Source/USDStageEditor/Private/USDStageEditorModule.cpp b/Engine/Plugins/Importers/USDImporter/Source/USDStageEditor/Private/USDStageEditorModule.cpp
--- a/Engine/Plugins/Importers/USDImporter/Source/USDStageEditor/Private/USDStageEditorModule.cpp
+++ b/Engine/Plugins/Importers/USDImporter/Source/USDStageEditor/Private/USDStageEditorModule.cpp
@@ -5,8 +5,8 @@
 #include "SUSDSaveDialog.h"
 #include "SUSDStage.h"
 #include "SUSDStageEditorStyle.h"
+#include "USDErrorUtils.h"
 #include "USDLayerUtils.h"
-#include "USDLog.h"
 #include "USDMemory.h"
 #include "USDProjectSettings.h"
 #include "USDStageActor.h"
@@ -211,7 +211,7 @@ namespace UE::UsdStageEditorModule::Private
 
 				if (!bSaved)
 				{
-					UE_LOG(LogUsd, Warning, TEXT("Failed to save layer '%s'"), *PinnedLayer.GetIdentifier());
+					USD_LOG_WARNING(TEXT("Failed to save layer '%s'"), *PinnedLayer.GetIdentifier());
 				}
 			}
 		}
@@ -227,7 +227,7 @@ namespace UE::UsdStageEditorModule::Private
 		}
 
 		const FText Title = LOCTEXT("MismatchedTypeNamesText", "USD: Mismatched type names");
-		UE_LOG(LogUsd, Warning, TEXT("%s"), *DialogText.ToString().Replace(TEXT("\n\n"), TEXT(" ")));
+		USD_LOG_USERWARNING(FText::FromString(DialogText.ToString().Replace(TEXT("\n\n"), TEXT(" "))));
 
 		FScopedUnrealAllocs UEAllocs;
 
diff --git a/Engine/Plugins/Importers/USDImporter/Source/USDStageEditorViewModels/Private/USDObjectFieldViewModel.cpp b/Engine/Plugins/Importers/USDImporter/Source/USDStageEditorViewModels/Private/USDObjectFieldViewModel.cpp
--- a/Engine/Plugins/Importers/USDImporter/Source/USDStageEditorViewModels/Private/USDObjectFieldViewModel.cpp
+++ b/Engine/Plugins/Importers/USDImporter/Source/USDStageEditorViewModels/Private/USDObjectFieldViewModel.cpp
@@ -298,7 +298,7 @@ void FUsdObjectFieldsViewModel::SetFieldValue(const FString& FieldName, const Us
 		ErrorToast.Image = FCoreStyle::Get().GetBrush(TEXT("MessageLog.Warning"));
 		FSlateNotificationManager::Get().AddNotification(ErrorToast);
 
-		FUsdLogManager::LogMessage(EMessageSeverity::Warning, ErrorMessage);
+		USD_LOG_USERWARNING(ErrorMessage);
 	}
 #endif	  // #if USE_USD_SDK
 }
diff --git a/Engine/Plugins/Importers/USDImporter/Source/USDStageEditorViewModels/Private/USDReferencesViewModel.cpp b/Engine/Plugins/Importers/USDImporter/Source/USDStageEditorViewModels/Private/USDReferencesViewModel.cpp
--- a/Engine/Plugins/Importers/USDImporter/Source/USDStageEditorViewModels/Private/USDReferencesViewModel.cpp
+++ b/Engine/Plugins/Importers/USDImporter/Source/USDStageEditorViewModels/Private/USDReferencesViewModel.cpp
@@ -2,7 +2,7 @@
 
 #include "USDReferencesViewModel.h"
 
-#include "USDLog.h"
+#include "USDErrorUtils.h"
 #include "USDMemory.h"
 #include "USDTypesConversion.h"
 
@@ -264,7 +264,7 @@ void FUsdReferencesViewModel::ReloadReference(const TSharedPtr<FUsdReference>& R
 		{
 			if (Layer)
 			{
-				UE_LOG(LogUsd, Log, TEXT("Reloading layer '%s'"), *UsdToUnreal::ConvertString(Layer->GetIdentifier()));
+				USD_LOG_INFO(TEXT("Reloading layer '%s'"), *UsdToUnreal::ConvertString(Layer->GetIdentifier()));
 				const bool bForce = true;
 				Layer->Reload(bForce);
 			}
diff --git a/Engine/Plugins/Importers/USDImporter/Source/USDStageEditorViewModels/Private/USDStageViewModel.cpp b/Engine/Plugins/Importers/USDImporter/Source/USDStageEditorViewModels/Private/USDStageViewModel.cpp
--- a/Engine/Plugins/Importers/USDImporter/Source/USDStageEditorViewModels/Private/USDStageViewModel.cpp
+++ b/Engine/Plugins/Importers/USDImporter/Source/USDStageEditorViewModels/Private/USDStageViewModel.cpp
@@ -8,7 +8,6 @@
 #include "USDConversionUtils.h"
 #include "USDErrorUtils.h"
 #include "USDLayerUtils.h"
-#include "USDLog.h"
 #include "USDMemory.h"
 #include "USDStageActor.h"
 #include "USDStageImportContext.h"
@@ -56,7 +55,7 @@ void FUsdStageViewModel::NewStage()
 {
 	FScopedTransaction Transaction(LOCTEXT("NewStageTransaction", "Created new USD stage"));
 
-	UsdUtils::StartMonitoringErrors();
+	FScopedUsdMessageLog ScopedLog;
 
 	if (!UsdStageActor.IsValid())
 	{
@@ -69,13 +68,11 @@ void FUsdStageViewModel::NewStage()
 		StageActor->Modify();
 		StageActor->NewStage();
 	}
-
-	UsdUtils::ShowErrorsAndStopMonitoring();
 }
 
 void FUsdStageViewModel::OpenStage(const TCHAR* FilePath)
 {
-	UsdUtils::StartMonitoringErrors();
+	FScopedUsdMessageLog ScopedLog;
 
 	if (!UsdStageActor.IsValid())
 	{
@@ -90,10 +87,8 @@ void FUsdStageViewModel::OpenStage(const TCHAR* FilePath)
 	}
 	else
 	{
-		UE_LOG(LogUsd, Error, TEXT("Failed to find a AUsdStageActor that could open stage '%s'!"), FilePath);
+		USD_LOG_ERROR(TEXT("Failed to find a AUsdStageActor that could open stage '%s'!"), FilePath);
 	}
-
-	UsdUtils::ShowErrorsAndStopMonitoring();
 }
 
 void FUsdStageViewModel::SetLoadAllRule()
@@ -191,7 +186,8 @@ void FUsdStageViewModel::ReloadStage()
 
 	if (UsdStage)
 	{
-		UsdUtils::StartMonitoringErrors();
+		FScopedUsdMessageLog ScopedLog;
+
 		{
 			FScopedUsdAllocs Allocs;
 
@@ -229,7 +225,7 @@ void FUsdStageViewModel::ReloadStage()
 			UsdUtils::GetUEPersistentStateSublayer(Stage, bCreateIfNeeded);
 		}
 
-		if (UsdUtils::ShowErrorsAndStopMonitoring())
+		if (FUsdLogManager::HasAccumulatedErrors())
 		{
 			return;
 		}
@@ -287,7 +283,7 @@ void FUsdStageViewModel::SaveStage()
 		{
 			FScopedUsdAllocs UsdAllocs;
 
-			UsdUtils::StartMonitoringErrors();
+			FScopedUsdMessageLog ScopedLog;
 
 			// Save layers manually instead of calling UsdStage::Save(). This is roughly the same implementation anyway, except
 			// that UsdStage::Save() will ignore a layer if it also happens to be added as a sublayer to a session layer, and
@@ -301,8 +297,6 @@ void FUsdStageViewModel::SaveStage()
 			}
 
 			UsdViewModelImpl::SaveUEStateLayer(UsdStage);
-
-			UsdUtils::ShowErrorsAndStopMonitoring(LOCTEXT("USDSaveError", "Failed to save current USD Stage!\nCheck the Output Log for details."));
 		}
 	}
 #endif	  // #if USE_USD_SDK
@@ -318,7 +312,7 @@ void FUsdStageViewModel::SaveStageAs(const TCHAR* FilePath)
 	{
 		if (UE::FUsdStage UsdStage = UsdStageActor->GetOrOpenUsdStage())
 		{
-			UsdUtils::StartMonitoringErrors();
+			FScopedUsdMessageLog ScopedLog;
 
 			if (UE::FSdfLayer RootLayer = UsdStage.GetRootLayer())
 			{
@@ -362,9 +356,6 @@ void FUsdStageViewModel::SaveStageAs(const TCHAR* FilePath)
 					UsdViewModelImpl::SaveUEStateLayer(UsdStage);
 				}
 			}
-
-			UsdUtils::ShowErrorsAndStopMonitoring(LOCTEXT("USDSaveAsError", "Failed to SaveAs current USD Stage!\nCheck the Output Log for details.")
-			);
 		}
 	}
 #endif	  // #if USE_USD_SDK
diff --git a/Engine/Plugins/Importers/USDImporter/Source/USDStageImporter/Private/USDStageAssetImportFactory.cpp b/Engine/Plugins/Importers/USDImporter/Source/USDStageImporter/Private/USDStageAssetImportFactory.cpp
--- a/Engine/Plugins/Importers/USDImporter/Source/USDStageImporter/Private/USDStageAssetImportFactory.cpp
+++ b/Engine/Plugins/Importers/USDImporter/Source/USDStageImporter/Private/USDStageAssetImportFactory.cpp
@@ -159,7 +159,7 @@ EReimportResult::Type UUsdStageAssetImportFactory::Reimport(UObject* Obj)
 {
 	if (!Obj)
 	{
-		FUsdLogManager::LogMessage(EMessageSeverity::Error, LOCTEXT("ReimportErrorInvalidAsset", "Failed to reimport asset as it is invalid!"));
+		USD_LOG_USERERROR(LOCTEXT("ReimportErrorInvalidAsset", "Failed to reimport asset as it is invalid!"));
 		return EReimportResult::Failed;
 	}
 
@@ -186,13 +186,10 @@ EReimportResult::Type UUsdStageAssetImportFactory::Reimport(UObject* Obj)
 
 	if (ReimportFilePath.IsEmpty() || OriginalPrimPath.IsEmpty())
 	{
-		FUsdLogManager::LogMessage(
-			EMessageSeverity::Error,
-			FText::Format(
-				LOCTEXT("ReimportErrorNoImportData", "Failed to reimport asset '{0}' as it doesn't seem to have valid USD import data or user data!"),
-				FText::FromName(Obj->GetFName())
-			)
-		);
+		USD_LOG_USERERROR(FText::Format(
+			LOCTEXT("ReimportErrorNoImportData", "Failed to reimport asset '{0}' as it doesn't seem to have valid USD import data or user data!"),
+			FText::FromName(Obj->GetFName())
+		));
 		return EReimportResult::Failed;
 	}
 
@@ -210,13 +207,10 @@ EReimportResult::Type UUsdStageAssetImportFactory::Reimport(UObject* Obj)
 	const bool bAllowActorImport = false;
 	if (!ImportContext.Init(Obj->GetName(), ReimportFilePath, Obj->GetName(), Obj->GetFlags(), IsAutomatedImport(), bIsReimport, bAllowActorImport))
 	{
-		FUsdLogManager::LogMessage(
-			EMessageSeverity::Error,
-			FText::Format(
-				LOCTEXT("ReimportErrorNoContext", "Failed to initialize re-import context for asset '{0}'!"),
-				FText::FromName(Obj->GetFName())
-			)
-		);
+		USD_LOG_USERERROR(FText::Format(
+			LOCTEXT("ReimportErrorNoContext", "Failed to initialize re-import context for asset '{0}'!"),
+			FText::FromName(Obj->GetFName())
+		));
 		return EReimportResult::Cancelled;
 	}
 
diff --git a/Engine/Plugins/Importers/USDImporter/Source/USDStageImporter/Private/USDStageImporter.cpp b/Engine/Plugins/Importers/USDImporter/Source/USDStageImporter/Private/USDStageImporter.cpp
--- a/Engine/Plugins/Importers/USDImporter/Source/USDStageImporter/Private/USDStageImporter.cpp
+++ b/Engine/Plugins/Importers/USDImporter/Source/USDStageImporter/Private/USDStageImporter.cpp
@@ -2,6 +2,7 @@
 
 #include "USDStageImporter.h"
 
+#include "Objects/USDSchemaTranslator.h"
 #include "USDAssetCache3.h"
 #include "USDAssetImportData.h"
 #include "USDAssetUserData.h"
@@ -11,12 +12,10 @@
 #include "USDErrorUtils.h"
 #include "USDGeomMeshConversion.h"
 #include "USDInfoCacheObject.h"
-#include "USDLog.h"
 #include "USDObjectUtils.h"
 #include "USDPrimLinkCacheObject.h"
 #include "USDPrimTwin.h"
 #include "USDSchemasModule.h"
-#include "Objects/USDSchemaTranslator.h"
 #include "USDStageImportContext.h"
 #include "USDStageImportOptions.h"
 #include "USDTypesConversion.h"
@@ -82,7 +81,7 @@ namespace UE::USDStageImporter::Private
 			return;
 		}
 
-		UsdUtils::StartMonitoringErrors();
+		FScopedUsdMessageLog ScopedLog;
 
 		UE::FUsdStage Stage;
 		if (bNeedsMasking)
@@ -140,24 +139,10 @@ namespace UE::USDStageImporter::Private
 			);
 		}
 
-		TArray<FString> ErrorStrings = UsdUtils::GetErrorsAndStopMonitoring();
-		FString Error = FString::Join(ErrorStrings, TEXT("\n"));
-
-		if (Error.IsEmpty())
-		{
-			ImportContext.Stage = Stage;
-		}
-		else
+		if (FUsdLogManager::HasAccumulatedErrors())
 		{
 			ImportContext.Stage = UE::FUsdStage();
-			FUsdLogManager::LogMessage(
-				EMessageSeverity::Error,
-				FText::Format(
-					LOCTEXT("CouldNotImportUSDFile", "Could not import USD file {0}\n {1}"),
-					FText::FromString(FilePath),
-					FText::FromString(Error)
-				)
-			);
+			USD_LOG_USERERROR(FText::Format(LOCTEXT("CouldNotImportUSDFile", "Could not import USD file {0}"), FText::FromString(FilePath)));
 		}
 	}
 
@@ -624,14 +609,11 @@ namespace UE::USDStageImporter::Private
 		UPackage* Package = ExistingPackage ? ExistingPackage : CreatePackage(*TargetPackagePath);
 		if (!Package)
 		{
-			FUsdLogManager::LogMessage(
-				EMessageSeverity::Error,
-				FText::Format(
-					LOCTEXT("PublishFailure", "Failed to get destination package at '{0}' for imported asset '{1}'!"),
-					FText::FromString(TargetPackagePath),
-					FText::FromName(Asset->GetFName())
-				)
-			);
+			USD_LOG_USERERROR(FText::Format(
+				LOCTEXT("PublishFailure", "Failed to get destination package at '{0}' for imported asset '{1}'!"),
+				FText::FromString(TargetPackagePath),
+				FText::FromName(Asset->GetFName())
+			));
 			return nullptr;
 		}
 		Package->FullyLoad();
@@ -1196,10 +1178,7 @@ namespace UE::USDStageImporter::Private
 
 		if (!ImportContext.SceneActor)
 		{
-			FUsdLogManager::LogMessage(
-				EMessageSeverity::Error,
-				LOCTEXT("NoSceneActor", "Failed to publish actors as there was no scene actor available!")
-			);
+			USD_LOG_USERERROR(LOCTEXT("NoSceneActor", "Failed to publish actors as there was no scene actor available!"));
 			return;
 		}
 
@@ -1782,10 +1761,7 @@ void UUsdStageImporter::ImportFromFile(FUsdStageImportContext& ImportContext)
 #if USE_USD_SDK
 	if (!ImportContext.World)
 	{
-		FUsdLogManager::LogMessage(
-			EMessageSeverity::Error,
-			LOCTEXT("NoWorldError", "Failed to import USD Stage because the target UWorld is invalid!")
-		);
+		USD_LOG_USERERROR(LOCTEXT("NoWorldError", "Failed to import USD Stage because the target UWorld is invalid!"));
 		return;
 	}
 
@@ -1813,7 +1789,7 @@ void UUsdStageImporter::ImportFromFile(FUsdStageImportContext& ImportContext)
 
 	if (!ImportContext.Stage)
 	{
-		FUsdLogManager::LogMessage(EMessageSeverity::Error, LOCTEXT("NoStageError", "Failed to open the USD Stage!"));
+		USD_LOG_USERERROR(LOCTEXT("NoStageError", "Failed to open the USD Stage!"));
 		return;
 	}
 
@@ -2002,7 +1978,7 @@ void UUsdStageImporter::ImportFromFile(FUsdStageImportContext& ImportContext)
 	{
 		double ElapsedSeconds = FPlatformTime::ToSeconds64(FPlatformTime::Cycles64() - StartTime);
 		UE::USDStageImporter::Private::SendAnalytics(ImportContext, nullptr, TEXT("Import"), PublishedAssets, ElapsedSeconds);
-		UE_LOG(LogUsd, Log, TEXT("Imported '%s' in %.3f seconds."), *ImportContext.FilePath, ElapsedSeconds);
+		USD_LOG_INFO(TEXT("Imported '%s' in %.3f seconds."), *ImportContext.FilePath, ElapsedSeconds);
 
 		UsdUnreal::Analytics::CollectSchemaAnalytics(ImportContext.Stage, TEXT("Import"));
 	}
@@ -2039,7 +2015,7 @@ bool UUsdStageImporter::ReimportSingleAsset(
 
 	if (!ImportContext.Stage)
 	{
-		FUsdLogManager::LogMessage(EMessageSeverity::Error, LOCTEXT("NoStageError", "Failed to open the USD Stage!"));
+		USD_LOG_USERERROR(LOCTEXT("NoStageError", "Failed to open the USD Stage!"));
 		return bSuccess;
 	}
 
@@ -2257,7 +2233,7 @@ bool UUsdStageImporter::ReimportSingleAsset(
 	{
 		double ElapsedSeconds = FPlatformTime::ToSeconds64(FPlatformTime::Cycles64() - StartTime);
 		UE::USDStageImporter::Private::SendAnalytics(ImportContext, ReimportedObject, TEXT("Reimport"), {ReimportedObject}, ElapsedSeconds);
-		UE_LOG(LogUsd, Log, TEXT("Re-imported '%s' in %.3f seconds."), *ImportContext.FilePath, ElapsedSeconds);
+		USD_LOG_INFO(TEXT("Re-imported '%s' in %.3f seconds."), *ImportContext.FilePath, ElapsedSeconds);
 
 		UsdUnreal::Analytics::CollectSchemaAnalytics(ImportContext.Stage, TEXT("Reimport"));
 	}
diff --git a/Engine/Plugins/Importers/USDImporter/Source/USDTests/Private/USDTestsBlueprintLibrary.cpp b/Engine/Plugins/Importers/USDImporter/Source/USDTests/Private/USDTestsBlueprintLibrary.cpp
--- a/Engine/Plugins/Importers/USDImporter/Source/USDTests/Private/USDTestsBlueprintLibrary.cpp
+++ b/Engine/Plugins/Importers/USDImporter/Source/USDTests/Private/USDTestsBlueprintLibrary.cpp
@@ -3,7 +3,7 @@
 #include "USDTestsBlueprintLibrary.h"
 
 #include "Objects/USDInfoCache.h"
-#include "USDLog.h"
+#include "USDErrorUtils.h"
 #include "USDStageActor.h"
 
 #include "CoreMinimal.h"
@@ -37,21 +37,21 @@ bool USDTestsBlueprintLibrary::RecompileBlueprintStageActor(AUsdStageActor* Blue
 	FKismetEditorUtilities::CompileBlueprint(BP, EBlueprintCompileOptions::SkipGarbageCollection);
 	if (BP->Status == EBlueprintStatus::BS_UpToDate)
 	{
-		UE_LOG(LogUsd, Log, TEXT("Blueprint compiled successfully (%s)"), *BP->GetName());
+		USD_LOG_INFO(TEXT("Blueprint compiled successfully (%s)"), *BP->GetName());
 		return true;
 	}
 	else if (BP->Status == EBlueprintStatus::BS_UpToDateWithWarnings)
 	{
-		UE_LOG(LogUsd, Warning, TEXT("Blueprint compiled successfully with warnings(%s)"), *BP->GetName());
+		USD_LOG_WARNING(TEXT("Blueprint compiled successfully with warnings(%s)"), *BP->GetName());
 		return true;
 	}
 	else if (BP->Status == EBlueprintStatus::BS_Error)
 	{
-		UE_LOG(LogUsd, Error, TEXT("Blueprint failed to compile (%s)"), *BP->GetName());
+		USD_LOG_ERROR(TEXT("Blueprint failed to compile (%s)"), *BP->GetName());
 		return false;
 	}
 
-	UE_LOG(LogUsd, Error, TEXT("Blueprint is in an unexpected state after compiling (%s)"), *BP->GetName());
+	USD_LOG_ERROR(TEXT("Blueprint is in an unexpected state after compiling (%s)"), *BP->GetName());
 #endif	  // WITH_EDITOR
 	return false;
 }
@@ -83,7 +83,7 @@ void USDTestsBlueprintLibrary::DirtyStageActorBlueprint(AUsdStageActor* Blueprin
 	bool bSuccess = FBlueprintEditorUtils::AddMemberVariable(BP, VarName, StringPinType);
 	if (!bSuccess)
 	{
-		UE_LOG(LogUsd, Error, TEXT("Failed to add new variable to blueprint (%s)"), *BP->GetName());
+		USD_LOG_ERROR(TEXT("Failed to add new variable to blueprint (%s)"), *BP->GetName());
 	}
 #endif	  // WITH_EDITOR
 }
diff --git a/Engine/Plugins/Importers/USDImporterMDL/Source/USDImporterMDL/Private/MDLUSDShadeMaterialTranslator.cpp b/Engine/Plugins/Importers/USDImporterMDL/Source/USDImporterMDL/Private/MDLUSDShadeMaterialTranslator.cpp
--- a/Engine/Plugins/Importers/USDImporterMDL/Source/USDImporterMDL/Private/MDLUSDShadeMaterialTranslator.cpp
+++ b/Engine/Plugins/Importers/USDImporterMDL/Source/USDImporterMDL/Private/MDLUSDShadeMaterialTranslator.cpp
@@ -218,13 +218,10 @@ void FMdlUsdShadeMaterialTranslator::CreateAssets()
 		}
 		else if (!MdlMaterial)
 		{
-			FUsdLogManager::LogMessage(
-				EMessageSeverity::Warning,
-				FText::Format(
-					LOCTEXT("UsdMdlConversionFailed", "Failed to create MDL material for prim {0}."),
-					FText::FromString(PrimPath.GetString())
-				)
-			);
+			USD_LOG_USERWARNING(FText::Format(
+				LOCTEXT("UsdMdlConversionFailed", "Failed to create MDL material for prim {0}."),
+				FText::FromString(PrimPath.GetString())
+			));
 		}
 
 		// Create material instance
diff --git a/Engine/Plugins/Runtime/USDCore/Source/USDClasses/Public/USDProjectSettings.h b/Engine/Plugins/Runtime/USDCore/Source/USDClasses/Public/USDProjectSettings.h
--- a/Engine/Plugins/Runtime/USDCore/Source/USDClasses/Public/USDProjectSettings.h
+++ b/Engine/Plugins/Runtime/USDCore/Source/USDClasses/Public/USDProjectSettings.h
@@ -81,6 +81,10 @@ public:
 	UPROPERTY(config, EditAnywhere, Category = USD)
 	bool bLogUsdSdkErrors = true;
 
+	// Similar messages are merged into one. Disable this option to see the individual messages (may slow performance and use more memory).
+	UPROPERTY(config, EditAnywhere, Category = USD)
+	bool bOptimizeUsdLog = true;
+
 	/**
 	 * USD Asset Cache to use for USD Stage Actors that don't have any asset cache specified.
 	 * Leave this empty to have each stage actor generate it's on transient cache instead.
diff --git a/Engine/Plugins/Runtime/USDCore/Source/USDUtilities/Private/Objects/USDInfoCache.cpp b/Engine/Plugins/Runtime/USDCore/Source/USDUtilities/Private/Objects/USDInfoCache.cpp
--- a/Engine/Plugins/Runtime/USDCore/Source/USDUtilities/Private/Objects/USDInfoCache.cpp
+++ b/Engine/Plugins/Runtime/USDCore/Source/USDUtilities/Private/Objects/USDInfoCache.cpp
@@ -3,8 +3,8 @@
 #include "Objects/USDInfoCache.h"
 
 #include "Objects/USDSchemaTranslator.h"
+#include "USDErrorUtils.h"
 #include "USDGeomMeshConversion.h"
-#include "USDLog.h"
 #include "USDMemory.h"
 #include "USDTypesConversion.h"
 
@@ -35,6 +35,8 @@
 #include "USDIncludesEnd.h"
 #endif	  // USE_USD_SDK
 
+#define LOCTEXT_NAMESPACE "UsdInfoCache"
+
 static int32 GMaxNumVerticesCollapsedMesh = 5000000;
 static FAutoConsoleVariableRef CVarMaxNumVerticesCollapsedMesh(
 	TEXT("USD.MaxNumVerticesCollapsedMesh"),
@@ -434,8 +436,6 @@ public:
 				FWriteScopeLock Lock = LockForWriting(*AuxPrim);
 				AuxPrim->MainPrims.Add(MainPrimPath);
 			}
-
-			UE_LOG(LogUsd, Verbose, TEXT("Registering main prim '%s' and aux prim '%s'"), *MainPrimPath.GetString(), *AuxPrimPath.GetString());
 		}
 	}
 
@@ -465,8 +465,6 @@ public:
 				FWriteScopeLock Lock = LockForWriting(*User);
 				User->AuxPrims.Add(MaterialPath);
 			}
-
-			UE_LOG(LogUsd, Verbose, TEXT("Registering user '%s' of material prim '%s'"), *UserPath.GetString(), *MaterialPath.GetString());
 		}
 	}
 
@@ -1315,14 +1313,15 @@ namespace UE::USDInfoCache::Private
 							if (!CollectionPath || !MaterialPath || PathVector.size() != 2)
 							{
 								// Emit this warning here as USD doesn't seem to and just seems to just ignores this relationship instead
-								UE_LOG(
-									LogUsd,
-									Warning,
-									TEXT("Prim '%s' describes a collection-based material binding, but the relationship '%s' is invalid: It should "
-										 "contain exactly one Material path and one path to a collection relationship"),
-									*Info.PrimPath.GetString(),
-									*UsdToUnreal::ConvertToken(Rel.GetName())
-								);
+								USD_LOG_USERWARNING(FText::Format(
+									LOCTEXT(
+										"InvalidCollection",
+										"Prim '{0}' describes a collection-based material binding, but the relationship '{1}' is invalid: It should "
+										"contain exactly one Material path and one path to a collection relationship"
+									),
+									FText::FromString(Info.PrimPath.GetString()),
+									FText::FromString(UsdToUnreal::ConvertToken(Rel.GetName()))
+								));
 								continue;
 							}
 
@@ -1338,13 +1337,14 @@ namespace UE::USDInfoCache::Private
 							}
 							else
 							{
-								UE_LOG(
-									LogUsd,
-									Warning,
-									TEXT("Failed to find collection at path '%s' when processing collection-based material bindings on prim '%s'"),
-									*UsdToUnreal::ConvertPath(CollectionPath->GetPrimPath()),
-									*Info.PrimPath.GetString()
-								);
+								USD_LOG_USERWARNING(FText::Format(
+									LOCTEXT(
+										"MissingCollection",
+										"Failed to find collection at path '{0}' when processing collection-based material bindings on prim '{1}'"
+									),
+									FText::FromString(UsdToUnreal::ConvertPath(CollectionPath->GetPrimPath())),
+									FText::FromString(Info.PrimPath.GetString())
+								));
 							}
 						}
 					}
@@ -2195,15 +2195,16 @@ namespace UE::USDInfoCache::Private
 		{
 			if (Depth > GGeometryCacheMaxDepth)
 			{
-				UE_LOG(
-					LogUsd,
-					Warning,
-					TEXT("Prim '%s' is potentially a geometry cache %d levels deep, which exceeds the limit of %d. "
-						 "This could affect its imported animation. The limit can be increased with the cvar USD.GeometryCache.MaxDepth if needed."),
-					*Info.PrimPath.GetString(),
+				USD_LOG_USERWARNING(FText::Format(
+					LOCTEXT(
+						"DeepGeometryCache",
+						"Prim '{0}' is potentially a geometry cache {1} levels deep, which exceeds the limit of {2}. "
+						"This could affect its imported animation. The limit can be increased with the cvar USD.GeometryCache.MaxDepth if needed."
+					),
+					FText::FromString(Info.PrimPath.GetString()),
 					Depth,
 					GGeometryCacheMaxDepth
-				);
+				));
 			}
 			FindValidGeometryCacheRoot(UsdPrim, Context, Impl, PrimState);
 			Depth = -1;
@@ -2268,14 +2269,15 @@ namespace UE::USDInfoCache::Private
 		{
 			if (Depth > GGeometryCacheMaxDepth)
 			{
-				UE_LOG(
-					LogUsd,
-					Warning,
-					TEXT("The stage has a geometry cache %d levels deep, which exceeds the limit of %d. "
-						 "This could affect its imported animation. The limit can be increased with the cvar USD.GeometryCache.MaxDepth if needed."),
+				USD_LOG_USERWARNING(FText::Format(
+					LOCTEXT(
+						"DeepGeometryCacheInStage",
+						"The stage has a geometry cache {0} levels deep, which exceeds the limit of {1}. "
+						"This could affect its imported animation. The limit can be increased with the cvar USD.GeometryCache.MaxDepth if needed."
+					),
 					Depth,
 					GGeometryCacheMaxDepth
-				);
+				));
 			}
 
 			// The pseudoroot itself cannot be a root for the geometry cache so start from its children
@@ -2709,9 +2711,7 @@ TOptional<bool> FUsdInfoCache::CanXformableSubtreeBeCollapsed(const UE::FSdfPath
 			return bCanBeCollapsed;
 		}
 
-		UE_LOG(
-			LogUsd,
-			Warning,
+		USD_LOG_WARNING(
 			TEXT(
 				"Failed to find whether subtree '%s' can be collapsed or not. Note: This function is meant to be used only during the main FUsdInfoCache build!"
 			),
@@ -2722,3 +2722,5 @@ TOptional<bool> FUsdInfoCache::CanXformableSubtreeBeCollapsed(const UE::FSdfPath
 
 	return {};
 }
+
+#undef LOCTEXT_NAMESPACE
diff --git a/Engine/Plugins/Runtime/USDCore/Source/USDUtilities/Private/Objects/USDPrimLinkCache.cpp b/Engine/Plugins/Runtime/USDCore/Source/USDUtilities/Private/Objects/USDPrimLinkCache.cpp
--- a/Engine/Plugins/Runtime/USDCore/Source/USDUtilities/Private/Objects/USDPrimLinkCache.cpp
+++ b/Engine/Plugins/Runtime/USDCore/Source/USDUtilities/Private/Objects/USDPrimLinkCache.cpp
@@ -2,7 +2,6 @@
 
 #include "Objects/USDPrimLinkCache.h"
 
-#include "USDLog.h"
 #include "UsdWrappers/SdfPath.h"
 
 #include "Misc/ScopeRWLock.h"
@@ -68,8 +67,6 @@ void FUsdPrimLinkCache::LinkAssetToPrim(const UE::FSdfPath& Path, UObject* Asset
 
 	FWriteScopeLock ScopeLock(ImplPtr->PrimPathToAssetsLock);
 
-	UE_LOG(LogUsd, Verbose, TEXT("Linking asset '%s' to prim '%s'"), *Asset->GetPathName(), *Path.GetString());
-
 	ImplPtr->PrimPathToAssets.FindOrAdd(Path).AddUnique(Asset);
 	ImplPtr->AssetToPrimPaths.FindOrAdd(Asset).AddUnique(Path);
 }
@@ -84,8 +81,6 @@ void FUsdPrimLinkCache::UnlinkAssetFromPrim(const UE::FSdfPath& Path, UObject* A
 
 	FWriteScopeLock ScopeLock(ImplPtr->PrimPathToAssetsLock);
 
-	UE_LOG(LogUsd, Verbose, TEXT("Unlinking asset '%s' to prim '%s'"), *Asset->GetPathName(), *Path.GetString());
-
 	if (TArray<TWeakObjectPtr<UObject>>* FoundAssetsForPrim = ImplPtr->PrimPathToAssets.Find(Path))
 	{
 		FoundAssetsForPrim->Remove(Asset);
@@ -106,8 +101,6 @@ TArray<TWeakObjectPtr<UObject>> FUsdPrimLinkCache::RemoveAllAssetPrimLinks(const
 
 	FWriteScopeLock ScopeLock(ImplPtr->PrimPathToAssetsLock);
 
-	UE_LOG(LogUsd, Verbose, TEXT("Removing asset prim links for path '%s'"), *Path.GetString());
-
 	TArray<TWeakObjectPtr<UObject>> Assets;
 	ImplPtr->PrimPathToAssets.RemoveAndCopyValue(Path, Assets);
 
@@ -132,8 +125,6 @@ TArray<UE::FSdfPath> FUsdPrimLinkCache::RemoveAllAssetPrimLinks(const UObject* A
 
 	FWriteScopeLock ScopeLock(ImplPtr->PrimPathToAssetsLock);
 
-	UE_LOG(LogUsd, Verbose, TEXT("Removing asset prim links for asset '%s'"), *Asset->GetPathName());
-
 	TArray<UE::FSdfPath> PrimPaths;
 	ImplPtr->AssetToPrimPaths.RemoveAndCopyValue(const_cast<UObject*>(Asset), PrimPaths);
 
@@ -158,8 +149,6 @@ void FUsdPrimLinkCache::RemoveAllAssetPrimLinks()
 
 	FWriteScopeLock ScopeLock(ImplPtr->PrimPathToAssetsLock);
 
-	UE_LOG(LogUsd, Verbose, TEXT("Removing all asset prim links"));
-
 	ImplPtr->PrimPathToAssets.Empty();
 	ImplPtr->AssetToPrimPaths.Empty();
 }
diff --git a/Engine/Plugins/Runtime/USDCore/Source/USDUtilities/Private/USDAttributeUtils.cpp b/Engine/Plugins/Runtime/USDCore/Source/USDUtilities/Private/USDAttributeUtils.cpp
--- a/Engine/Plugins/Runtime/USDCore/Source/USDUtilities/Private/USDAttributeUtils.cpp
+++ b/Engine/Plugins/Runtime/USDCore/Source/USDUtilities/Private/USDAttributeUtils.cpp
@@ -3,8 +3,8 @@
 #include "USDAttributeUtils.h"
 
 #include "UnrealUSDWrapper.h"
+#include "USDErrorUtils.h"
 #include "USDLayerUtils.h"
-#include "USDLog.h"
 #include "USDMemory.h"
 #include "USDProjectSettings.h"
 #include "USDTypesConversion.h"
@@ -243,7 +243,7 @@ void UsdUtils::NotifyIfOverriddenOpinion(const pxr::UsdProperty& Property)
 				FText::FromString(UsdToUnreal::ConvertString(SpecLayer->GetIdentifier()))
 			);
 
-			UE_LOG(LogUsd, Warning, TEXT("%s"), *SubText.ToString().Replace(TEXT("\n\n"), TEXT(" ")));
+			USD_LOG_USERWARNING(FText::FromString(SubText.ToString().Replace(TEXT("\n\n"), TEXT(" "))));
 
 			const UUsdProjectSettings* Settings = GetDefault<UUsdProjectSettings>();
 			if (Settings && Settings->bShowOverriddenOpinionsWarning)
@@ -367,22 +367,12 @@ bool UsdUtils::NotifyIfInstanceProxy(const pxr::UsdPrim& Prim)
 			{
 				case EUsdEditInInstanceBehavior::Ignore:
 				{
-					UE_LOG(
-						LogUsd,
-						Log,
-						TEXT("Ignoring some edits to prim '%s' as it is an instance proxy"),
-						*UsdToUnreal::ConvertPath(Prim.GetPrimPath())
-					);
+					USD_LOG_INFO(TEXT("Ignoring some edits to prim '%s' as it is an instance proxy"), *UsdToUnreal::ConvertPath(Prim.GetPrimPath()));
 					return true;
 				}
 				case EUsdEditInInstanceBehavior::RemoveInstanceable:
 				{
-					UE_LOG(
-						LogUsd,
-						Log,
-						TEXT("Removing all instanceable flags from ancestors of prim '%s'"),
-						*UsdToUnreal::ConvertPath(Prim.GetPrimPath())
-					);
+					USD_LOG_INFO(TEXT("Removing all instanceable flags from ancestors of prim '%s'"), *UsdToUnreal::ConvertPath(Prim.GetPrimPath()));
 					RemoveInstanceables();
 
 					// We shouldn't be instanceable now, so we can probably author whatever we wanted to author
@@ -416,7 +406,7 @@ bool UsdUtils::NotifyIfInstanceProxy(const pxr::UsdPrim& Prim)
 						FText::FromString(FirstInstancePath)
 					);
 
-					UE_LOG(LogUsd, Warning, TEXT("%s"), *SubText.ToString().Replace(TEXT("\n\n"), TEXT(" ")));
+					USD_LOG_USERWARNING(FText::FromString(SubText.ToString().Replace(TEXT("\n\n"), TEXT(" "))));
 
 					static TWeakPtr<SNotificationItem> Notification;
 
diff --git a/Engine/Plugins/Runtime/USDCore/Source/USDUtilities/Private/USDConversionUtils.cpp b/Engine/Plugins/Runtime/USDCore/Source/USDUtilities/Private/USDConversionUtils.cpp
--- a/Engine/Plugins/Runtime/USDCore/Source/USDUtilities/Private/USDConversionUtils.cpp
+++ b/Engine/Plugins/Runtime/USDCore/Source/USDUtilities/Private/USDConversionUtils.cpp
@@ -11,7 +11,6 @@
 #include "USDGeomMeshConversion.h"
 #include "USDIntegrationUtils.h"
 #include "USDLayerUtils.h"
-#include "USDLog.h"
 #include "USDMemory.h"
 #include "USDObjectUtils.h"
 #include "USDProjectSettings.h"
@@ -177,19 +176,16 @@ namespace USDConversionUtilsImpl
 		if (IgnoredPrimvarNames.Num() > 0)
 		{
 			FString PrimvarNames = FString::Join(IgnoredPrimvarNames, TEXT(", "));
-			FUsdLogManager::LogMessage(
-				EMessageSeverity::Warning,
-				FText::Format(
-					LOCTEXT(
-						"TooHighUVIndex",
-						"Mesh '{0}' has some valid UV set primvars ({1}) that will be ignored because they target an UV index larger than the "
-						"highest supported ({2})"
-					),
-					FText::FromString(MeshPath),
-					FText::FromString(PrimvarNames),
-					MAX_STATIC_TEXCOORDS - 1
-				)
-			);
+			USD_LOG_USERWARNING(FText::Format(
+				LOCTEXT(
+					"TooHighUVIndex",
+					"Mesh '{0}' has some valid UV set primvars ({1}) that will be ignored because they target an UV index larger than the "
+					"highest supported ({2})"
+				),
+				FText::FromString(MeshPath),
+				FText::FromString(PrimvarNames),
+				MAX_STATIC_TEXCOORDS - 1
+			));
 		}
 
 		// Show a warning if the mesh does not contain the exact primvars the material wants
@@ -217,14 +213,11 @@ namespace USDConversionUtilsImpl
 
 			if (!bFoundUsablePrimvar)
 			{
-				FUsdLogManager::LogMessage(
-					EMessageSeverity::Warning,
-					FText::Format(
-						LOCTEXT("DidNotFindPrimvar", "Could not find primvar '{0}' on mesh '{1}', used by its bound material"),
-						FText::FromString(UsdToUnreal::ConvertString(UsedPrimvar.GetBaseName())),
-						FText::FromString(MeshPath)
-					)
-				);
+				USD_LOG_USERWARNING(FText::Format(
+					LOCTEXT("DidNotFindPrimvar", "Could not find primvar '{0}' on mesh '{1}', used by its bound material"),
+					FText::FromString(UsdToUnreal::ConvertString(UsedPrimvar.GetBaseName())),
+					FText::FromString(MeshPath)
+				));
 			}
 		}
 
@@ -253,20 +246,17 @@ namespace USDConversionUtilsImpl
 					UsedPrimvar = &Primvars[0];
 				}
 
-				FUsdLogManager::LogMessage(
-					EMessageSeverity::Warning,
-					FText::Format(
-						LOCTEXT(
-							"MoreThanOnePrimvarForIndex",
-							"Mesh '{0}' has more than one primvar used as UV set with index '{1}'. The UV set will use the values from primvar "
-							"'{2}'{3}"
-						),
-						FText::FromString(MeshPath),
-						UVAndPrimvars.Key,
-						FText::FromString(UsdToUnreal::ConvertString(UsedPrimvar->GetBaseName())),
-						bUsedByMaterial ? FText::FromString(TEXT(", as its used by its bound material")) : FText::GetEmpty()
-					)
-				);
+				USD_LOG_USERWARNING(FText::Format(
+					LOCTEXT(
+						"MoreThanOnePrimvarForIndex",
+						"Mesh '{0}' has more than one primvar used as UV set with index '{1}'. The UV set will use the values from primvar "
+						"'{2}'{3}"
+					),
+					FText::FromString(MeshPath),
+					UVAndPrimvars.Key,
+					FText::FromString(UsdToUnreal::ConvertString(UsedPrimvar->GetBaseName())),
+					bUsedByMaterial ? FText::FromString(TEXT(", as its used by its bound material")) : FText::GetEmpty()
+				));
 			}
 		}
 	}
@@ -303,7 +293,7 @@ namespace USDConversionUtilsImpl
 			FText::FromString(PrimNamesString)
 		);
 
-		UE_LOG(LogUsd, Warning, TEXT("%s"), *SubText.ToString().Replace(TEXT("\n\n"), TEXT(" ")));
+		USD_LOG_USERWARNING(FText::FromString(SubText.ToString().Replace(TEXT("\n\n"), TEXT(" "))));
 
 		const UUsdProjectSettings* Settings = GetDefault<UUsdProjectSettings>();
 		if (Settings && Settings->bShowWarningOnIncompleteDuplication)
@@ -1906,16 +1896,15 @@ TMap<FString, UsdUtils::FVolumePrimInfo> UsdUtils::GetVolumeInfoByFilePathHash(c
 											ExistingFormat = &ExistingIter->second;
 										}
 
-										UE_LOG(
-											LogUsd,
-											Warning,
-											TEXT(
-												"OpenVDBAsset prims disagree on the attribute channel format for the Sparse Volume Texture generated for VDB file '%s' (encountered '%s' and '%s'). If there are multiple opinions for the attribute channel formats from different OpenVDBAsset prims, they must all agree!"
+										USD_LOG_USERWARNING(FText::Format(
+											LOCTEXT(
+												"DisagreeAttributeChannel",
+												"OpenVDBAsset prims disagree on the attribute channel format for the Sparse Volume Texture generated for VDB file '{0}' (encountered '{1}' and '{2}'). If there are multiple opinions for the attribute channel formats from different OpenVDBAsset prims, they must all agree!"
 											),
-											*ResolvedVDBPath,
-											*UsdToUnreal::ConvertToken(DataType),
-											ExistingFormat ? **ExistingFormat : TEXT("unknown")
-										);
+											FText::FromString(ResolvedVDBPath),
+											FText::FromString(UsdToUnreal::ConvertToken(DataType)),
+											FText::FromString(ExistingFormat ? **ExistingFormat : TEXT("unknown"))
+										));
 									}
 									else if (!bIsSet)
 									{
@@ -1924,15 +1913,13 @@ TMap<FString, UsdUtils::FVolumePrimInfo> UsdUtils::GetVolumeInfoByFilePathHash(c
 								}
 								else
 								{
-									UE_LOG(
-										LogUsd,
-										Warning,
-										TEXT(
-											"Invalid Sparse Volume Texture attribute channel format '%s'. Available formats: 'unorm8', 'float16' and 'float32'."
+									USD_LOG_USERWARNING(FText::Format(
+										LOCTEXT(
+											"InvalidChannelFormat",
+											"Invalid Sparse Volume Texture attribute channel format '{0}'. Available formats: 'unorm8', 'float16' and 'float32'."
 										),
-										*UsdToUnreal::ConvertToken(DataType),
-										*UsdToUnreal::ConvertPath(AssetPrimPath)
-									);
+										FText::FromString(UsdToUnreal::ConvertToken(DataType))
+									));
 								}
 							}
 						};
@@ -1963,18 +1950,17 @@ TMap<FString, UsdUtils::FVolumePrimInfo> UsdUtils::GetVolumeInfoByFilePathHash(c
 									{
 										if (Component != *ExistingComponentMapping)
 										{
-											UE_LOG(
-												LogUsd,
-												Warning,
-												TEXT(
-													"Found multiple OpenVDBAsset prims (including '%s') targetting the same grid '%s', but with with conflicting grid component to Sparse Volume Texture attribute channel mapping (for example, both components '%s' and '%s' are mapped to the same channel '%s', which is not allowed)"
+											USD_LOG_USERWARNING(FText::Format(
+												LOCTEXT(
+													"MultipleTargettingSameGrid",
+													"Found multiple OpenVDBAsset prims (including '{0}') targetting the same grid '{1}', but with with conflicting grid component to Sparse Volume Texture attribute channel mapping (for example, both components '{2}' and '{3}' are mapped to the same channel '{4}', which is not allowed)"
 												),
-												*UsdToUnreal::ConvertPath(AssetPrimPath),
-												*GridNameStr,
-												**ExistingComponentMapping,
-												*Component,
-												*Channel
-											);
+												FText::FromString(UsdToUnreal::ConvertPath(AssetPrimPath)),
+												FText::FromString(GridNameStr),
+												FText::FromString(*ExistingComponentMapping),
+												FText::FromString(Component),
+												FText::FromString(Channel)
+											));
 										}
 									}
 									else
@@ -1985,18 +1971,17 @@ TMap<FString, UsdUtils::FVolumePrimInfo> UsdUtils::GetVolumeInfoByFilePathHash(c
 							}
 							else
 							{
-								UE_LOG(
-									LogUsd,
-									Warning,
-									TEXT(
-										"Failed to parse custom component to attribute mapping from OpenVDBAsset prim '%s': The '%s' and '%s' attributes should have the same number of entries, but the former has %u entries while the latter has %u"
+								USD_LOG_USERWARNING(FText::Format(
+									LOCTEXT(
+										"FailCustomAttributeMapping",
+										"Failed to parse custom component to attribute mapping from OpenVDBAsset prim '{0}': The '{1}' and '{2}' attributes should have the same number of entries, but the former has {3} entries while the latter has {4}"
 									),
-									*UsdToUnreal::ConvertPath(AssetPrimPath),
-									*UsdToUnreal::ConvertToken(UnrealIdentifiers::UnrealSVTMappedGridComponents),
-									*UsdToUnreal::ConvertToken(UnrealIdentifiers::UnrealSVTMappedAttributeChannels),
-									DesiredComponents.size(),
-									DesiredChannels.size()
-								);
+									FText::FromString(UsdToUnreal::ConvertPath(AssetPrimPath)),
+									FText::FromString(UsdToUnreal::ConvertToken(UnrealIdentifiers::UnrealSVTMappedGridComponents)),
+									FText::FromString(UsdToUnreal::ConvertToken(UnrealIdentifiers::UnrealSVTMappedAttributeChannels)),
+									static_cast<uint64>(DesiredComponents.size()),
+									static_cast<uint64>(DesiredChannels.size())
+								));
 							}
 						}
 					}
@@ -2004,25 +1989,21 @@ TMap<FString, UsdUtils::FVolumePrimInfo> UsdUtils::GetVolumeInfoByFilePathHash(c
 			}
 			else
 			{
-				UE_LOG(
-					LogUsd,
-					Warning,
-					TEXT("Failed to find the VDB file '%s' referenced by OpenVDBAsset prim at path '%s'"),
-					*ResolvedVDBPath,
-					*UsdToUnreal::ConvertPath(AssetPrimPath)
-				);
+				USD_LOG_USERWARNING(FText::Format(
+					LOCTEXT("FailToFindFile", "Failed to find the VDB file '{0}' referenced by OpenVDBAsset prim at path '{1}'"),
+					FText::FromString(ResolvedVDBPath),
+					FText::FromString(UsdToUnreal::ConvertPath(AssetPrimPath))
+				));
 			}
 		}
 		else
 		{
-			UE_LOG(
-				LogUsd,
-				Warning,
-				TEXT("Failed to find an OpenVDBAsset prim at path '%s' for field '%s' of prim '%s'"),
-				*UsdToUnreal::ConvertPath(AssetPrimPath),
-				*UsdToUnreal::ConvertToken(FieldName),
-				*UsdToUnreal::ConvertPath(Volume.GetPrim().GetPath())
-			);
+			USD_LOG_USERWARNING(FText::Format(
+				LOCTEXT("FailToFindPrim", "Failed to find an OpenVDBAsset prim at path '{0}' for field '{1}' of prim '{2}'"),
+				FText::FromString(UsdToUnreal::ConvertPath(AssetPrimPath)),
+				FText::FromString(UsdToUnreal::ConvertToken(FieldName)),
+				FText::FromString(UsdToUnreal::ConvertPath(Volume.GetPrim().GetPath()))
+			));
 		}
 	}
 
@@ -2059,18 +2040,17 @@ TMultiMap<FString, FString> UsdUtils::GetVolumeMaterialParameterToFieldNameMap(c
 
 	if (MappedFields.size() != MappedParameters.size())
 	{
-		UE_LOG(
-			LogUsd,
-			Warning,
-			TEXT(
-				"Failed to parse custom parsed texture to material parameter mapping from volume prim '%s': The '%s' and '%s' attributes should have the same number of entries, but the former has %u entries while the latter has %u"
+		USD_LOG_USERWARNING(FText::Format(
+			LOCTEXT(
+				"FailToParseMaterialMapping",
+				"Failed to parse custom parsed texture to material parameter mapping from volume prim '{0}': The '{1}' and '{2}' attributes should have the same number of entries, but the former has {3} entries while the latter has {4}"
 			),
-			*UsdToUnreal::ConvertPath(VolumePrim.GetPrimPath()),
-			*UsdToUnreal::ConvertToken(UnrealIdentifiers::UnrealSVTMappedFields),
-			*UsdToUnreal::ConvertToken(UnrealIdentifiers::UnrealSVTMappedMaterialParameters),
-			MappedFields.size(),
-			MappedParameters.size()
-		);
+			FText::FromString(UsdToUnreal::ConvertPath(VolumePrim.GetPrimPath())),
+			FText::FromString(UsdToUnreal::ConvertToken(UnrealIdentifiers::UnrealSVTMappedFields)),
+			FText::FromString(UsdToUnreal::ConvertToken(UnrealIdentifiers::UnrealSVTMappedMaterialParameters)),
+			static_cast<uint64>(MappedFields.size()),
+			static_cast<uint64>(MappedParameters.size())
+		));
 		return {};
 	}
 
@@ -2444,9 +2424,7 @@ UUsdAssetUserData* UsdUtils::GetAssetUserData(const UObject* Object, TSubclassOf
 	const IInterface_AssetUserData* AssetUserDataInterface = Cast<const IInterface_AssetUserData>(Object);
 	if (!AssetUserDataInterface)
 	{
-		UE_LOG(
-			LogUsd,
-			Warning,
+		USD_LOG_WARNING(
 			TEXT("Tried getting AssetUserData from object '%s', but the class '%s' doesn't implement the AssetUserData interface!"),
 			*Object->GetPathName(),
 			*Object->GetClass()->GetName()
@@ -2473,9 +2451,7 @@ UUsdAssetUserData* UsdUtils::GetOrCreateAssetUserData(UObject* Object, TSubclass
 	IInterface_AssetUserData* AssetUserDataInterface = Cast<IInterface_AssetUserData>(Object);
 	if (!AssetUserDataInterface)
 	{
-		UE_LOG(
-			LogUsd,
-			Warning,
+		USD_LOG_WARNING(
 			TEXT("Tried adding AssetUserData to object '%s', but it doesn't implement the AssetUserData interface!"),
 			*Object->GetPathName()
 		);
@@ -2505,9 +2481,7 @@ bool UsdUtils::SetAssetUserData(UObject* Object, UUsdAssetUserData* AssetUserDat
 	IInterface_AssetUserData* AssetUserDataInterface = Cast<IInterface_AssetUserData>(Object);
 	if (!AssetUserDataInterface)
 	{
-		UE_LOG(
-			LogUsd,
-			Warning,
+		USD_LOG_WARNING(
 			TEXT("Tried adding AssetUserData to object '%s', but it doesn't implement the AssetUserData interface!"),
 			*Object->GetPathName()
 		);
@@ -2516,7 +2490,6 @@ bool UsdUtils::SetAssetUserData(UObject* Object, UUsdAssetUserData* AssetUserDat
 
 	while (AssetUserDataInterface->HasAssetUserDataOfClass(UUsdAssetUserData::StaticClass()))
 	{
-		UE_LOG(LogUsd, Verbose, TEXT("Removing old AssetUserData from object '%s' before adding a new one"), *Object->GetPathName());
 		AssetUserDataInterface->RemoveUserDataOfClass(UUsdAssetUserData::StaticClass());
 	}
 
@@ -2649,16 +2622,15 @@ namespace UE::UsdConversionUtils::Private
 
 				if (bUnknownTargetType)
 				{
-					UE_LOG(
-						LogUsd,
-						Warning,
-						TEXT(
-							"Failed to match the referenced type when adding a reference or payload to prim '%s', as the target prim spec '%s' has an unknown type '%s'! The referencer type will be cleared instead."
+					USD_LOG_USERWARNING(FText::Format(
+						LOCTEXT(
+							"MatchReferencerFail",
+							"Failed to match the referenced type when adding a reference or payload to prim '{0}', as the target prim spec '{1}' has an unknown type '{2}'! The referencer type will be cleared instead."
 						),
-						*UsdToUnreal::ConvertPath(ReferencerPrim.GetPrimPath()),
-						*UsdToUnreal::ConvertPath(TargetPrimSpec->GetPath()),
-						*UsdToUnreal::ConvertToken(TargetTypeName)
-					);
+						FText::FromString(UsdToUnreal::ConvertPath(ReferencerPrim.GetPrimPath())),
+						FText::FromString(UsdToUnreal::ConvertPath(TargetPrimSpec->GetPath())),
+						FText::FromString(UsdToUnreal::ConvertToken(TargetTypeName))
+					));
 					ReferencerPrim.ClearTypeName();
 				}
 				else
@@ -2975,9 +2947,7 @@ bool UsdUtils::RenamePrim(UE::FUsdPrim& Prim, const TCHAR* NewPrimName)
 
 	if (!bCanApply)
 	{
-		UE_LOG(
-			LogUsd,
-			Error,
+		USD_LOG_ERROR(
 			TEXT("Failed to rename prim with path '%s' to name '%s'. Errors:\n%s"),
 			*Prim.GetPrimPath().GetString(),
 			NewPrimName,
@@ -3004,9 +2974,7 @@ bool UsdUtils::RenamePrim(UE::FUsdPrim& Prim, const TCHAR* NewPrimName)
 			if (!Layer->Apply(Edit))
 			{
 				// This should not be happening since CanApply was true, so stop doing whatever it is we're doing
-				UE_LOG(
-					LogUsd,
-					Error,
+				USD_LOG_ERROR(
 					TEXT("Failed to rename prim with path '%s' to name '%s' in layer '%s'"),
 					*Prim.GetPrimPath().GetString(),
 					NewPrimName,
@@ -3350,9 +3318,7 @@ USDUTILITIES_API void UsdUtils::RemoveAllLocalPrimSpecs(const UE::FUsdPrim& Prim
 			continue;
 		}
 
-		UE_LOG(
-			LogUsd,
-			Log,
+		USD_LOG_INFO(
 			TEXT("Removing prim spec '%s' from layer '%s'"),
 			*UsdToUnreal::ConvertPath(SpecPath),
 			*UsdToUnreal::ConvertString(SpecLayer->GetIdentifier())
@@ -3511,7 +3477,7 @@ bool UsdUtils::CopyPrims(const TArray<UE::FUsdPrim>& Prims)
 		}
 
 		bCopiedSomething = true;
-		UE_LOG(LogUsd, Log, TEXT("Copied prim '%s' into the clipboard"), *Prim.GetPrimPath().GetString());
+		USD_LOG_INFO(TEXT("Copied prim '%s' into the clipboard"), *Prim.GetPrimPath().GetString());
 	}
 #endif	  // USE_USD_SDK
 
@@ -3597,9 +3563,7 @@ TArray<UE::FSdfPath> UsdUtils::PastePrims(const UE::FUsdPrim& ParentPrim)
 			continue;
 		}
 
-		UE_LOG(
-			LogUsd,
-			Log,
+		USD_LOG_INFO(
 			TEXT("Pasted prim '%s' as a child of prim '%s' within the edit target '%s'"),
 			*OriginalName,
 			*UsdToUnreal::ConvertPath(UsdParentPrim.GetPath()),
@@ -3876,9 +3840,7 @@ TArray<UE::FSdfPath> UsdUtils::DuplicatePrims(const TArray<UE::FUsdPrim>& Prims,
 
 			if (!pxr::SdfJustCreatePrimInLayer(UsdLayer, NewSpecPath))
 			{
-				UE_LOG(
-					LogUsd,
-					Warning,
+				USD_LOG_WARNING(
 					TEXT("Failed to create prim and parent specs for path '%s' within layer '%s'"),
 					*UsdToUnreal::ConvertPath(NewSpecPath),
 					*UsdToUnreal::ConvertString(UsdLayer->GetIdentifier())
@@ -3888,9 +3850,7 @@ TArray<UE::FSdfPath> UsdUtils::DuplicatePrims(const TArray<UE::FUsdPrim>& Prims,
 
 			if (!pxr::SdfCopySpec(FlattenedLayer, FlattenedPrim->GetPath(), UsdLayer, NewSpecPath))
 			{
-				UE_LOG(
-					LogUsd,
-					Warning,
+				USD_LOG_WARNING(
 					TEXT("Failed to copy flattened prim spec from '%s' onto path '%s' within layer '%s'"),
 					*UsdToUnreal::ConvertPath(UsdPrim.GetPath()),
 					*UsdToUnreal::ConvertPath(NewSpecPath),
@@ -3899,9 +3859,7 @@ TArray<UE::FSdfPath> UsdUtils::DuplicatePrims(const TArray<UE::FUsdPrim>& Prims,
 				return Result;
 			}
 
-			UE_LOG(
-				LogUsd,
-				Log,
+			USD_LOG_INFO(
 				TEXT("Flattened prim '%s' onto spec '%s' at layer '%s'"),
 				*UsdToUnreal::ConvertPath(UsdPrim.GetPath()),
 				*UsdToUnreal::ConvertPath(NewSpecPath),
@@ -3915,9 +3873,7 @@ TArray<UE::FSdfPath> UsdUtils::DuplicatePrims(const TArray<UE::FUsdPrim>& Prims,
 				pxr::SdfPath SpecPath = Spec->GetPath();
 				pxr::SdfLayerHandle SpecLayerHandle = Spec->GetLayer();
 
-				UE_LOG(
-					LogUsd,
-					Log,
+				USD_LOG_INFO(
 					TEXT("Duplicating prim spec '%s' within layer '%s'"),
 					*UsdToUnreal::ConvertPath(SpecPath),
 					*UsdToUnreal::ConvertString(SpecLayerHandle->GetIdentifier())
@@ -3932,9 +3888,7 @@ TArray<UE::FSdfPath> UsdUtils::DuplicatePrims(const TArray<UE::FUsdPrim>& Prims,
 				// Let's ensure that though, just in case
 				if (!pxr::SdfJustCreatePrimInLayer(SpecLayerHandle, NewSpecPath))
 				{
-					UE_LOG(
-						LogUsd,
-						Warning,
+					USD_LOG_WARNING(
 						TEXT("Failed to create prim and parent specs for path '%s' within layer '%s'"),
 						*UsdToUnreal::ConvertPath(NewSpecPath),
 						*UsdToUnreal::ConvertString(SpecLayerHandle->GetIdentifier())
@@ -3974,9 +3928,7 @@ TArray<UE::FSdfPath> UsdUtils::DuplicatePrims(const TArray<UE::FUsdPrim>& Prims,
 				// the destination spec before copying stuff, and we may want to copy multiple specs overwriting each other
 				if (!pxr::SdfCopySpec(SpecLayerHandle, SpecPath, SpecLayerHandle, NewSpecPath, ShouldCopyValue, ShouldCopyChildren))
 				{
-					UE_LOG(
-						LogUsd,
-						Warning,
+					USD_LOG_WARNING(
 						TEXT("Failed to copy spec from path '%s' onto path '%s' within layer '%s'"),
 						*UsdToUnreal::ConvertPath(SpecPath),
 						*UsdToUnreal::ConvertPath(NewSpecPath),
@@ -4156,9 +4108,7 @@ bool UsdUtils::ClearNonEssentialPrimMetadata(const pxr::UsdPrim& Prim)
 
 		if (!bSuccess)
 		{
-			UE_LOG(
-				LogUsd,
-				Warning,
+			USD_LOG_WARNING(
 				TEXT("Failed to clear metadata field '%s' from prim '%s'"),
 				*UsdToUnreal::ConvertToken(FieldName),
 				*UsdToUnreal::ConvertPath(Prim.GetPrimPath())
diff --git a/Engine/Plugins/Runtime/USDCore/Source/USDUtilities/Private/USDErrorUtils.cpp b/Engine/Plugins/Runtime/USDCore/Source/USDUtilities/Private/USDErrorUtils.cpp
--- a/Engine/Plugins/Runtime/USDCore/Source/USDUtilities/Private/USDErrorUtils.cpp
+++ b/Engine/Plugins/Runtime/USDCore/Source/USDUtilities/Private/USDErrorUtils.cpp
@@ -2,11 +2,14 @@
 
 #include "USDErrorUtils.h"
 
+#include "USDLog.h"
+#include "USDProjectSettings.h"
+
+#include "HAL/IConsoleManager.h"
 #include "Math/NumericLimits.h"
 #include "Misc/ScopeLock.h"
 #include "Modules/ModuleManager.h"
-
-#include "USDLog.h"
+#include "Templates/SharedPointer.h"
 
 #if WITH_EDITOR
 #include "IMessageLogListing.h"
@@ -17,10 +20,6 @@
 #include "USDMemory.h"
 #include "USDTypesConversion.h"
 
-#include "Framework/Notifications/NotificationManager.h"
-#include "Templates/SharedPointer.h"
-#include "Widgets/Notifications/SNotificationList.h"
-
 #include "USDIncludesStart.h"
 #include "pxr/base/tf/errorMark.h"
 #include "USDIncludesEnd.h"
@@ -28,16 +27,93 @@
 
 #define LOCTEXT_NAMESPACE "USDErrorUtils"
 
-#if USE_USD_SDK
+static bool GUsdUseMessageLog = true;
+static FAutoConsoleVariableRef CVarUsdUseMessageLog(
+	TEXT("USD.UseMessageLog"),
+	GUsdUseMessageLog,
+	TEXT(
+		"Output user-facing messages from USD code and the USD-related Unreal plugins on the Message Log (in addition to the Output Log, which is always done). This can help bring issues to attention and it should make it easier to visualize the messages, but it can be slower and consume more memory."
+	)
+);
 
-namespace UsdUtils
+namespace UE::USDErrorUtils::Private
 {
+#if USE_USD_SDK
 	// We need an extra level of indirection because TfErrorMark is noncopyable
 	using MarkRef = TUsdStore<TSharedRef<pxr::TfErrorMark>>;
+	TArray<MarkRef> ErrorMarkStack;
+	TUniquePtr<class FUsdDiagnosticDelegate> DiagnosticDelegate;
+	TUniquePtr<class FUsdCombinedMessages> CombinedMessages;
+	int32 CombinedMessagesRefCount;
+	FCriticalSection CombinedMessagesLock;
+
+	FName GenerateMessageIdentifier(uint32 Number)
+	{
+		return FName(TEXT("USD_LOG_ID"), Number);
+	}
+
+	void SendMessageToOutputLog(const FTokenizedMessage& Message)
+	{
+		if (Message.GetSeverity() == EMessageSeverity::Error)
+		{
+			UE_LOG(LogUsd, Error, TEXT("%s"), *(Message.ToText().ToString()));
+		}
+		else if (Message.GetSeverity() == EMessageSeverity::Warning || Message.GetSeverity() == EMessageSeverity::PerformanceWarning)
+		{
+			UE_LOG(LogUsd, Warning, TEXT("%s"), *(Message.ToText().ToString()));
+		}
+		else
+		{
+			UE_LOG(LogUsd, Log, TEXT("%s"), *(Message.ToText().ToString()));
+		}
+	}
 
-	static TArray<MarkRef> ErrorMarkStack;
+	void LogUsdDiagnosticMessage(const pxr::TfDiagnosticBase& DiagMessage)
+	{
+		FScopedUsdAllocs Allocs;
+
+		std::string Msg = DiagMessage.GetDiagnosticCodeAsString();
+		Msg += ": ";
+		Msg += DiagMessage.GetCommentary();
+
+		const size_t LineNumber = DiagMessage.GetSourceLineNumber();
+		const std::string Filename = DiagMessage.GetSourceFileName();
+		const pxr::TfEnum ErrorCode = DiagMessage.GetDiagnosticCode();
+
+		EMessageSeverity::Type Severity = EMessageSeverity::Info;
+		switch (ErrorCode.GetValueAsInt())
+		{
+			default:
+			case pxr::TF_DIAGNOSTIC_STATUS_TYPE:
+			{
+				Severity = EMessageSeverity::Info;
+				break;
+			}
+			case pxr::TF_DIAGNOSTIC_WARNING_TYPE:
+			{
+				Severity = EMessageSeverity::Warning;
+				break;
+			}
+			case pxr::TF_DIAGNOSTIC_FATAL_CODING_ERROR_TYPE:
+			case pxr::TF_DIAGNOSTIC_NONFATAL_ERROR_TYPE:
+			case pxr::TF_DIAGNOSTIC_RUNTIME_ERROR_TYPE:
+			case pxr::TF_DIAGNOSTIC_FATAL_ERROR_TYPE:
+			case pxr::TF_DIAGNOSTIC_CODING_ERROR_TYPE:
+			{
+				Severity = EMessageSeverity::Error;
+				break;
+			}
+		}
+
+		FScopedUnrealAllocs UEAllocs;
+
+		const uint32 Identifier = HashCombine(FCrc::StrCrc32(Filename.c_str()), GetTypeHash(LineNumber));
+
+		// Note: Using an FText here causes these messages to be user-facing
+		FUsdLogManager::Log(Severity, FText::FromString(UsdToUnreal::ConvertString(Msg)), Identifier);
+	}
 
-	void StartMonitoringErrors()
+	void PushErrorMark()
 	{
 		FScopedUsdAllocs Allocs;
 
@@ -47,207 +123,406 @@ namespace UsdUtils
 		ErrorMarkStack.Emplace(Mark);
 	}
 
-	TArray<FString> GetErrorsAndStopMonitoring()
+	void PopErrorMark()
 	{
 		if (ErrorMarkStack.Num() == 0)
 		{
-			return {};
+			return;
 		}
 
+		FScopedUsdAllocs Allocs;
+
 		MarkRef Store = ErrorMarkStack.Pop();
 		pxr::TfErrorMark& Mark = Store.Get().Get();
 
 		if (Mark.IsClean())
 		{
-			return {};
+			return;
 		}
 
-		TArray<FString> Errors;
-
 		for (pxr::TfErrorMark::Iterator ErrorIter = Mark.GetBegin(); ErrorIter != Mark.GetEnd(); ++ErrorIter)
 		{
-			std::string ErrorStr = ErrorIter->GetErrorCodeAsString();
-			ErrorStr += ": ";
-			ErrorStr += ErrorIter->GetCommentary();
-
-			// Add unique here as for some errors (e.g. parsing errors) USD can emit the exact same
-			// error message 5+ times in a row
-			Errors.AddUnique(UsdToUnreal::ConvertString(ErrorStr));
+			LogUsdDiagnosticMessage(*ErrorIter);
 		}
 
 		Mark.Clear();
-
-		return Errors;
 	}
 
-	bool ShowErrorsAndStopMonitoring(const FText& ToastMessage)
+	class FUsdDiagnosticDelegate : public pxr::TfDiagnosticMgr::Delegate
 	{
-		TArray<FString> Errors = GetErrorsAndStopMonitoring();
-		bool bHadErrors = Errors.Num() > 0;
+	public:
+		virtual ~FUsdDiagnosticDelegate() override = default;
 
-		if (bHadErrors)
+		virtual void IssueError(const pxr::TfError& Message) override
 		{
-			FNotificationInfo ErrorToast(
-				!ToastMessage.IsEmpty() ? ToastMessage : LOCTEXT("USDErrorsToast", "Encountered USD errors!\nCheck the Output Log for details.")
-			);
-
-			ErrorToast.ExpireDuration = 5.0f;
-			ErrorToast.bFireAndForget = true;
-			ErrorToast.Image = FCoreStyle::Get().GetBrush(TEXT("MessageLog.Error"));
-			FSlateNotificationManager::Get().AddNotification(ErrorToast);
+			LogUsdDiagnosticMessage(Message);
 		}
 
-		for (const FString& Error : Errors)
+		virtual void IssueFatalError(const pxr::TfCallContext& Context, const std::string& Message) override
 		{
-			FUsdLogManager::LogMessage(FTokenizedMessage::Create(EMessageSeverity::Error, FText::FromString(Error)));
-		}
+			FScopedUnrealAllocs UEAllocs;
 
-		return Errors.Num() > 0;
-	}
-};		  // namespace UsdUtils
+			EMessageSeverity::Type Severity = EMessageSeverity::Error;
+			const char* SourceFile = Context.GetFile();
+			size_t LineNumber = Context.GetLine();
 
-#else	  // #if USE_USD_SDK
+			uint32 Identifier = HashCombine(FCrc::StrCrc32(SourceFile), GetTypeHash(LineNumber));
+			FUsdLogManager::Log(Severity, FText::FromString(UsdToUnreal::ConvertString(Message)), Identifier);
+		}
 
-namespace UsdUtils
-{
-	void StartMonitoringErrors()
-	{
-	}
+		virtual void IssueStatus(const pxr::TfStatus& Message) override
+		{
+			LogUsdDiagnosticMessage(Message);
+		}
 
-	TArray<FString> GetErrorsAndStopMonitoring()
-	{
-		return {TEXT("USD SDK is not available!")};
-	}
+		virtual void IssueWarning(const pxr::TfWarning& Message) override
+		{
+			LogUsdDiagnosticMessage(Message);
+		}
+	};
 
-	bool ShowErrorsAndStopMonitoring(const FText& ToastMessage)
+	class FUsdCombinedMessages
 	{
-		return false;
-	}
-}
+	public:
+		FUsdCombinedMessages(const FUsdCombinedMessages& Other) = delete;
+		FUsdCombinedMessages(FUsdCombinedMessages&& Other) = default;
+		FUsdCombinedMessages& operator=(const FUsdCombinedMessages& Other) = delete;
+		FUsdCombinedMessages& operator=(FUsdCombinedMessages&& Other) = default;
 
-#endif	  // #if USE_USD_SDK
-
-namespace UE
-{
-	namespace Internal
-	{
-		FUsdMessageLog::~FUsdMessageLog()
+		FUsdCombinedMessages()
 		{
-			Dump();
+			const UUsdProjectSettings* ProjectSettings = GetDefault<UUsdProjectSettings>();
+			bOptimizeUsdLog = ProjectSettings->bOptimizeUsdLog;
 		}
 
-#if WITH_EDITOR
-		void FUsdMessageLog::Push(const TSharedRef<FTokenizedMessage>& Message)
+		~FUsdCombinedMessages()
 		{
-			TokenizedMessages.Add(Message);
+			DisplayMessages();
 		}
-#endif	  // WITH_EDITOR
 
-		void FUsdMessageLog::Dump()
+		void DisplayMessages()
 		{
 #if WITH_EDITOR
-			FMessageLogModule& MessageLogModule = FModuleManager::LoadModuleChecked<FMessageLogModule>("MessageLog");
-			TSharedRef<IMessageLogListing> LogListing = MessageLogModule.GetLogListing(TEXT("USD"));
+			if (bHasSkippedMessages)
+			{
+				const static TSharedRef<FTokenizedMessage> SkippedMessage = FTokenizedMessage::Create(
+					EMessageSeverity::Info,
+					LOCTEXT(
+						"SkippedMEssages",
+						"Some similar log messages were skipped during the previous USD operation. You can disable this behavior by setting the USD.CollapseLogMessages cvar to false."
+					)
+				);
+				SendMessageToOutputLog(*SkippedMessage);
+			}
+
+			if (!GUsdUseMessageLog)
+			{
+				return;
+			}
+
+			// Move our messages to a simple array that the IMessageLogListing can consume
+			TArray<TSharedRef<FTokenizedMessage>> TokenizedMessages;
+			if (bOptimizeUsdLog)
+			{
+				TokenizedMessages.Reserve(LoggedMessagesByID.Num());
+
+				for (const TPair<FName, FMessageInfo>& Pair : LoggedMessagesByID)
+				{
+					const FMessageInfo& MessageInfo = Pair.Value;
+					if (!MessageInfo.bUserFacing)
+					{
+						continue;
+					}
+
+					TSharedPtr<FTokenizedMessage> Message = MessageInfo.Message;
+					if (MessageInfo.Count > 1)
+					{
+						Message->AddText(FText::Format(LOCTEXT("InstancesTextEditor", " (and {0} similar messages)"), MessageInfo.Count));
+					}
+
+					TokenizedMessages.Add(Message.ToSharedRef());
+				}
+			}
+			else
+			{
+				TokenizedMessages.Reserve(UncollapsedLoggedMessages.Num());
+
+				for (const FMessageInfo& MessageInfo : UncollapsedLoggedMessages)
+				{
+					if (!MessageInfo.bUserFacing)
+					{
+						continue;
+					}
+
+					TokenizedMessages.Add(MessageInfo.Message.ToSharedRef());
+				}
+			}
 
 			if (TokenizedMessages.Num() > 0)
 			{
-				// This crashes internally at runtime due to Slate not finding the correct brush...
-				// I think MessageLog is not supposed to be used at runtime, even though it technically can.
-				// At runtime we shouldn't ever push, but just in case
-				LogListing->AddMessages(TokenizedMessages);
-				LogListing->NotifyIfAnyMessages(LOCTEXT("Log", "There were some issues loading the USD Stage."), EMessageSeverity::Info);
-				TokenizedMessages.Empty();
+				FMessageLogModule& MessageLogModule = FModuleManager::LoadModuleChecked<FMessageLogModule>("MessageLog");
+				TSharedRef<IMessageLogListing> LogListing = MessageLogModule.GetLogListing(TEXT("USD"));
+
+				// We'll output the messages to the output log ourselves from LogMessageInternal()
+				const bool bMirrorToOutputLog = false;
+				LogListing->AddMessages(TokenizedMessages, bMirrorToOutputLog);
+
+				// Force display even if all we have are info level messages.
+				// This is also handy because it also outputs them on the output log for us
+				const bool bForce = true;
+				LogListing->NotifyIfAnyMessages(	//
+					LOCTEXT("Log", "The previous USD operation produced some log messages."),
+					EMessageSeverity::Info,
+					bForce
+				);
 			}
 #endif	  // WITH_EDITOR
 		}
-	}	  // namespace Internal
-}	 // namespace UE
-
-TOptional<UE::Internal::FUsdMessageLog> FUsdLogManager::MessageLog;
-int32 FUsdLogManager::MessageLogRefCount;
-FCriticalSection FUsdLogManager::MessageLogLock;
 
-TSet<FString> LoggedMessages;
+		// Returns whether we already had a similar message previously added
+		bool AppendMessage(const TSharedRef<FTokenizedMessage>& Message, bool bUserFacing)
+		{
+			const FName Identifier = Message->GetIdentifier();
+			bool bFoundSimilarMessage = false;
 
-void FUsdLogManager::LogMessage(EMessageSeverity::Type Severity, const FText& Message)
-{
-	LogMessage(FTokenizedMessage::Create(Severity, Message));
-}
+			if (bOptimizeUsdLog)
+			{
+				// Use this instead of just FindOrAdd because it's probably more useful to retain the very first
+				// version of the message, if we end up getting multiple
+				if (FMessageInfo* FoundMessageInfo = LoggedMessagesByID.Find(Identifier))
+				{
+					bFoundSimilarMessage = true;
+					bHasSkippedMessages = true;
+
+					FoundMessageInfo->bUserFacing = FoundMessageInfo->bUserFacing || bUserFacing;
+					FoundMessageInfo->Count += 1;
+				}
+				else
+				{
+					FMessageInfo MessageInfo;
+					MessageInfo.Message = Message;
+					MessageInfo.bUserFacing = MessageInfo.bUserFacing || bUserFacing;
+					MessageInfo.Count += 1;
+
+					LoggedMessagesByID.Add(Identifier, MessageInfo);
+				}
+			}
+			else if (GUsdUseMessageLog)
+			{
+				FMessageInfo& MessageInfo = UncollapsedLoggedMessages.Emplace_GetRef();
+				MessageInfo.Message = Message;
+				MessageInfo.bUserFacing = bUserFacing;
+			}
 
-void FUsdLogManager::LogMessage(const TSharedRef<FTokenizedMessage>& Message)
-{
-	bool bMessageProcessed = false;
+			return bFoundSimilarMessage;
+		}
 
-	{
-		FScopeLock Lock(&MessageLogLock);
-		const FString& Str = Message->ToText().ToString();
-		if (LoggedMessages.Contains(Str))
+		bool HasErrorsOrWarnings()
 		{
-			return;
-		}
+			for (const TPair<FName, FMessageInfo>& Pair : LoggedMessagesByID)
+			{
+				// Higher severity has lower numbers
+				if ((int)Pair.Value.Message->GetSeverity() <= (int)EMessageSeverity::Warning)
+				{
+					return true;
+				}
+			}
 
-		LoggedMessages.Add(Str);
+			return false;
+		}
 
-#if WITH_EDITOR
-		if (MessageLog)
+	private:
+		struct FMessageInfo
 		{
-			MessageLog->Push(Message);
-			bMessageProcessed = true;
-		}
-#endif
-	}
+			TSharedPtr<FTokenizedMessage> Message;
+			bool bUserFacing = false;
+			int32 Count = 0;
+		};
+
+		bool bOptimizeUsdLog = false;
 
-	if (!bMessageProcessed)
+		bool bHasSkippedMessages = false;
+		TMap<FName, FMessageInfo> LoggedMessagesByID;
+		TArray<FMessageInfo> UncollapsedLoggedMessages;
+	};
+
+	void LogMessageInternal(const TSharedRef<FTokenizedMessage>& Message, bool bUserFacing)
 	{
-		if (Message->GetSeverity() == EMessageSeverity::Error)
-		{
-			UE_LOG(LogUsd, Error, TEXT("%s"), *(Message->ToText().ToString()));
-		}
-		else if (Message->GetSeverity() == EMessageSeverity::Warning || Message->GetSeverity() == EMessageSeverity::PerformanceWarning)
+		bool bFoundSimilar = false;
 		{
-			UE_LOG(LogUsd, Warning, TEXT("%s"), *(Message->ToText().ToString()));
+			FScopeLock Lock(&CombinedMessagesLock);
+			if (CombinedMessages)
+			{
+				bFoundSimilar = CombinedMessages->AppendMessage(Message, bUserFacing);
+			}
 		}
-		else
+
+		if (!bFoundSimilar)
 		{
-			UE_LOG(LogUsd, Log, TEXT("%s"), *(Message->ToText().ToString()));
+			// Log immediately, so that our messages are interleaved with any other non-USD logged messages
+			// naturally, and are present on the output log already in case we crash after this.
+			//
+			// We'll skip messages here if they're similar, but if we do it's because we have a CombinedMessages,
+			// and when that struct is destroyed it will output an additional message to the Output Log letting
+			// the user know that some messages were skipped and what they can do about it
+			SendMessageToOutputLog(*Message);
 		}
 	}
+#endif	  // USE_USD_SDK
+}	 // namespace UE::USDErrorUtils::Private
+
+void FUsdLogManager::Log(EMessageSeverity::Type Severity, const FString& Message, uint32 MessageID)
+{
+#if USE_USD_SDK
+	using namespace UE::USDErrorUtils::Private;
+
+	TSharedRef<FTokenizedMessage> TokenizedMessage = FTokenizedMessage::Create(Severity, FText::FromString(Message));
+	TokenizedMessage->SetIdentifier(GenerateMessageIdentifier(MessageID));
+
+	const bool bUserFacing = false;
+	LogMessageInternal(TokenizedMessage, bUserFacing);
+#endif	  // USE_USD_SDK
 }
 
-void FUsdLogManager::EnableMessageLog()
+void FUsdLogManager::Log(EMessageSeverity::Type Severity, const FText& Message, uint32 MessageID)
 {
-	FScopeLock Lock(&MessageLogLock);
+#if USE_USD_SDK
+	using namespace UE::USDErrorUtils::Private;
 
-	LoggedMessages.Reset();
+	TSharedRef<FTokenizedMessage> TokenizedMessage = FTokenizedMessage::Create(Severity, Message);
+	TokenizedMessage->SetIdentifier(GenerateMessageIdentifier(MessageID));
+
+	const bool bUserFacing = true;
+	LogMessageInternal(TokenizedMessage, bUserFacing);
+#endif	  // USE_USD_SDK
+}
+
+bool FUsdLogManager::HasAccumulatedErrors()
+{
+#if USE_USD_SDK
+	using namespace UE::USDErrorUtils::Private;
 
-	if (++MessageLogRefCount == 1)
+	FScopeLock Lock(&CombinedMessagesLock);
+
+	if (CombinedMessages.IsValid())
 	{
-		MessageLog.Emplace();
-		UsdUtils::StartMonitoringErrors();
+		return CombinedMessages->HasErrorsOrWarnings();
 	}
+#endif	  // USE_USD_SDK
 
-	check(MessageLogRefCount < MAX_int32);
+	return false;
 }
 
-void FUsdLogManager::DisableMessageLog()
+void FUsdLogManager::RegisterDiagnosticDelegate()
+{
+#if USE_USD_SDK
+	using namespace UE::USDErrorUtils::Private;
+
+	if (DiagnosticDelegate.IsValid())
+	{
+		UnregisterDiagnosticDelegate();
+	}
+
+	DiagnosticDelegate = MakeUnique<FUsdDiagnosticDelegate>();
+
+	pxr::TfDiagnosticMgr& DiagMgr = pxr::TfDiagnosticMgr::GetInstance();
+	DiagMgr.AddDelegate(DiagnosticDelegate.Get());
+#endif	  // USE_USD_SDK
+}
+
+void FUsdLogManager::UnregisterDiagnosticDelegate()
 {
-	FScopeLock Lock(&MessageLogLock);
+#if USE_USD_SDK
+	using namespace UE::USDErrorUtils::Private;
 
-	if (--MessageLogRefCount == 0)
+	if (!DiagnosticDelegate.IsValid())
 	{
-		UsdUtils::ShowErrorsAndStopMonitoring();
-		MessageLog.Reset();
+		return;
 	}
+
+	pxr::TfDiagnosticMgr& DiagMgr = pxr::TfDiagnosticMgr::GetInstance();
+	DiagMgr.RemoveDelegate(DiagnosticDelegate.Get());
+
+	DiagnosticDelegate = nullptr;
+#endif	  // USE_USD_SDK
 }
 
 FScopedUsdMessageLog::FScopedUsdMessageLog()
 {
-	FUsdLogManager::EnableMessageLog();
+#if USE_USD_SDK
+	using namespace UE::USDErrorUtils::Private;
+
+	FScopeLock Lock(&CombinedMessagesLock);
+
+	if (++CombinedMessagesRefCount == 1)
+	{
+		CombinedMessages = MakeUnique<FUsdCombinedMessages>();
+		PushErrorMark();
+	}
+
+	check(CombinedMessagesRefCount < MAX_int32);
+#endif	  // USE_USD_SDK
 }
 
 FScopedUsdMessageLog::~FScopedUsdMessageLog()
 {
-	FUsdLogManager::DisableMessageLog();
+#if USE_USD_SDK
+	using namespace UE::USDErrorUtils::Private;
+
+	FScopeLock Lock(&CombinedMessagesLock);
+
+	if (--CombinedMessagesRefCount == 0)
+	{
+		PopErrorMark();
+		CombinedMessages.Reset();
+	}
+#endif	  // USE_USD_SDK
+}
+
+// Deprecated
+void UsdUtils::StartMonitoringErrors()
+{
+}
+
+// Deprecated
+TArray<FString> UsdUtils::GetErrorsAndStopMonitoring()
+{
+	return {};
+}
+
+// Deprecated
+bool UsdUtils::ShowErrorsAndStopMonitoring(const FText& ToastMessage)
+{
+	return false;
+}
+
+// Deprecated
+void FUsdLogManager::LogMessage(EMessageSeverity::Type Severity, const FText& Message)
+{
+#if USE_USD_SDK
+	Log(Severity, Message, 0);
+#endif	  // USE_USD_SDK
+}
+
+// Deprecated
+void FUsdLogManager::LogMessage(const TSharedRef<FTokenizedMessage>& Message)
+{
+#if USE_USD_SDK
+	using namespace UE::USDErrorUtils::Private;
+
+	const bool bUserFacing = true;
+	LogMessageInternal(Message, bUserFacing);
+#endif	  // USE_USD_SDK
+}
+
+// Deprecated
+void FUsdLogManager::EnableMessageLog()
+{
+}
+
+// Deprecated
+void FUsdLogManager::DisableMessageLog()
+{
 }
 
 #undef LOCTEXT_NAMESPACE
diff --git a/Engine/Plugins/Runtime/USDCore/Source/USDUtilities/Private/USDGeomMeshConversion.cpp b/Engine/Plugins/Runtime/USDCore/Source/USDUtilities/Private/USDGeomMeshConversion.cpp
--- a/Engine/Plugins/Runtime/USDCore/Source/USDUtilities/Private/USDGeomMeshConversion.cpp
+++ b/Engine/Plugins/Runtime/USDCore/Source/USDUtilities/Private/USDGeomMeshConversion.cpp
@@ -9,8 +9,8 @@
 #include "USDClassesModule.h"
 #include "USDConversionUtils.h"
 #include "USDDrawModeComponent.h"
+#include "USDErrorUtils.h"
 #include "USDLayerUtils.h"
-#include "USDLog.h"
 #include "USDMemory.h"
 #include "USDObjectUtils.h"
 #include "USDPrimConversion.h"
@@ -1313,7 +1313,8 @@ namespace UE::UsdGeomMeshConversion::Private
 		{
 			std::string PrimvarName = Primvar.GetPrimvarName().GetString();
 			if (PrimvarName == "displayColor" || PrimvarName == "displayOpacity" || PrimvarName == "st"
-				|| (Options.ImportPrimvars == FUsdMeshConversionOptions::EImportPrimvar::Bake && !Options.PrimvarNames.Contains(UTF8_TO_TCHAR(PrimvarName.c_str()))))
+				|| (Options.ImportPrimvars == FUsdMeshConversionOptions::EImportPrimvar::Bake
+					&& !Options.PrimvarNames.Contains(UTF8_TO_TCHAR(PrimvarName.c_str()))))
 			{
 				continue;
 			}
@@ -1351,7 +1352,7 @@ namespace UE::UsdGeomMeshConversion::Private
 			}
 			else
 			{
-				UE_LOG(LogUsd, Warning, TEXT("Storage of primvar type '%s' not yet supported."), ANSI_TO_TCHAR(PrimvarType.c_str()));
+				USD_LOG_WARNING(TEXT("Storage of primvar type '%s' not yet supported."), *UsdToUnreal::ConvertString(PrimvarType));
 			}
 
 			bool bDataCollected = false;
@@ -1846,7 +1847,7 @@ namespace UE::UsdGeomMeshConversion::Private
 
 		if (bInvalidIndices)
 		{
-			UE_LOG(LogUsd, Warning, TEXT("Invalid primvar indices encountered in ComputeFlattened"));
+			USD_LOG_WARNING(TEXT("Invalid primvar indices encountered in ComputeFlattened"));
 		}
 
 		return Result;
@@ -1896,28 +1897,26 @@ namespace UE::UsdGeomMeshConversion::Private
 		int32 TargetSubdivLevel = FMath::Max(0, FMath::Min(GMaxSubdivLevel, Options.SubdivisionLevel));
 		if (TargetSubdivLevel < Options.SubdivisionLevel)
 		{
-			UE_LOG(
-				LogUsd,
-				Warning,
-				TEXT("Max subdivision level was clamped to %d (controlled by the cvar '%s')"),
+			USD_LOG_USERWARNING(FText::Format(
+				LOCTEXT("ClampedSubdiv", "Max subdivision level was clamped to {0} (controlled by the cvar '{1}')"),
 				GMaxSubdivLevel,
-				*MaxUsdSubdivLevelCvarName
-			);
+				FText::FromString(MaxUsdSubdivLevelCvarName)
+			));
 		}
 		if (TargetSubdivLevel < 1)
 		{
-			UE_LOG(
-				LogUsd,
-				Log,
-				TEXT("Cancelling out of subdividing mesh '%s' due to target subdivision level being %d after clamping (it needs to be at least 1 for "
-					 "a round of subdivision)"),
-				*InOutMeshData.SourcePrimPath,
+			USD_LOG_USERWARNING(FText::Format(
+				LOCTEXT(
+					"CancellingSubdiv",
+					"Cancelling out of subdividing mesh '{0}' due to target subdivision level being {1} after clamping (it needs to be at least 1 for a round of subdivision)"
+				),
+				FText::FromString(InOutMeshData.SourcePrimPath),
 				TargetSubdivLevel
-			);
+			));
 			return false;
 		}
 
-		UE_LOG(LogUsd, Log, TEXT("Subdividing mesh '%s' to subdivision level %d"), *InOutMeshData.SourcePrimPath, TargetSubdivLevel);
+		USD_LOG_INFO(TEXT("Subdividing mesh '%s' to subdivision level %d"), *InOutMeshData.SourcePrimPath, TargetSubdivLevel);
 
 		// We need to track our faceVarying attributes when subdividing, so we'll give each one an unique
 		// index into that FaceVaryingTopologies array down below. We'll use this to track how many entries that array
@@ -3035,16 +3034,16 @@ namespace UE::UsdGeomMeshConversion::Private
 						}
 						else
 						{
-							UE_LOG(
-								LogUsd,
-								Warning,
-								TEXT("Trying to read UV at index %u from prim '%s' but the UV set %d only has %u values! Using zeros "
-									 "instead."),
-								ValueIndex,
-								*InMeshData.SourcePrimPath,
+							USD_LOG_USERWARNING(FText::Format(
+								LOCTEXT(
+									"WrongUVIndex",
+									"Trying to read UV at index {0} from prim '{1}' but the UV set {2} only has {3} values! Using zeros instead."
+								),
+								static_cast<uint64>(ValueIndex),
+								FText::FromString(InMeshData.SourcePrimPath),
 								UVSetIndex,
-								UVSet.UVs.size()
-							);
+								static_cast<uint64>(UVSet.UVs.size())
+							));
 						}
 
 						// Flip V for Unreal uv's which match directx
@@ -3071,15 +3070,15 @@ namespace UE::UsdGeomMeshConversion::Private
 							}
 							else
 							{
-								UE_LOG(
-									LogUsd,
-									Warning,
-									TEXT("Trying to read displayColor at index %u from prim '%s' but the prim only has %u values! Using "
-										 "zeros instead."),
-									ValueIndex,
-									*InMeshData.SourcePrimPath,
-									InMeshData.DisplayColors.size()
-								);
+								USD_LOG_USERWARNING(FText::Format(
+									LOCTEXT(
+										"WrongDisplayColorIndex",
+										"Trying to read displayColor at index {0} from prim '{1}' but the prim only has {2} values! Using zeros instead."
+									),
+									static_cast<uint64>(ValueIndex),
+									FText::FromString(InMeshData.SourcePrimPath),
+									static_cast<uint64>(InMeshData.DisplayColors.size())
+								));
 							}
 						}
 
@@ -3103,15 +3102,15 @@ namespace UE::UsdGeomMeshConversion::Private
 							}
 							else
 							{
-								UE_LOG(
-									LogUsd,
-									Warning,
-									TEXT("Trying to read displayOpacity at index %u from prim '%s' but the prim only has %u values! Using "
-										 "zeros instead."),
-									ValueIndex,
-									*InMeshData.SourcePrimPath,
-									InMeshData.DisplayColors.size()
-								);
+								USD_LOG_USERWARNING(FText::Format(
+									LOCTEXT(
+										"WrongDisplayOpacityIndex",
+										"Trying to read displayOpacity at index {0} from prim '{1}' but the prim only has {2} values! Using zeros instead."
+									),
+									static_cast<uint64>(ValueIndex),
+									FText::FromString(InMeshData.SourcePrimPath),
+									static_cast<uint64>(InMeshData.DisplayColors.size())
+								));
 							}
 						}
 					}
@@ -3174,16 +3173,16 @@ namespace UE::UsdGeomMeshConversion::Private
 									}
 									else
 									{
-										UE_LOG(
-											LogUsd,
-											Warning,
-											TEXT("Trying to read Primvar '%s' at index %u from prim '%s' but the prim only has %u values! Using "
-												 "zeros instead."),
-											*PrimvarName,
-											ValueIndex,
-											*InMeshData.SourcePrimPath,
-											InMeshData.DisplayColors.size()
-										);
+										USD_LOG_USERWARNING(FText::Format(
+											LOCTEXT(
+												"WrongPrimvarIndex",
+												"Trying to read Primvar '{0}' at index {1} from prim '{2}' but the prim only has {3} values! Using zeros instead."
+											),
+											FText::FromString(PrimvarName),
+											static_cast<uint64>(ValueIndex),
+											FText::FromString(InMeshData.SourcePrimPath),
+											static_cast<uint64>(InMeshData.DisplayColors.size())
+										));
 									}
 								}
 							},
@@ -3238,9 +3237,7 @@ namespace UE::UsdGeomMeshConversion::Private
 
 		if (NumPolygons > 0 && NumSkippedPolygons > 0)
 		{
-			UE_LOG(
-				LogUsd,
-				Warning,
+			USD_LOG_WARNING(
 				TEXT("Skipped %d out of %d faces when parsing the mesh for prim '%s', as those faces contained too many repeated vertex indices"),
 				NumSkippedPolygons,
 				NumPolygons,
@@ -3399,9 +3396,7 @@ bool UsdToUnreal::ConvertPointInstancerToMesh(
 			UE::FUsdPrim PrototypeUsdPrim = Stage.GetPrimAtPath(PrototypePath);
 			if (!PrototypeUsdPrim)
 			{
-				UE_LOG(
-					LogUsd,
-					Warning,
+				USD_LOG_WARNING(
 					TEXT("Failed to find prototype '%s' for PointInstancer '%s' within ConvertPointInstancerToMesh"),
 					*UsdToUnreal::ConvertPath(PrototypePath),
 					*UsdToUnreal::ConvertPath(PointInstancer.GetPrim().GetPrimPath())
@@ -3501,9 +3496,7 @@ bool UsdToUnreal::ConvertPointInstancerToMesh(
 		bool bSuccess = UsdUtils::GetPointInstancerTransforms(StageInfo, PointInstancer, PrototypeIndex, Options.TimeCode, InstanceTransforms);
 		if (!bSuccess)
 		{
-			UE_LOG(
-				LogUsd,
-				Error,
+			USD_LOG_ERROR(
 				TEXT("Failed to retrieve point instancer transforms for prototype index '%u' of point instancer '%s'"),
 				PrototypeIndex,
 				*UsdToUnreal::ConvertPath(PointInstancer.GetPrim().GetPrimPath())
@@ -3754,9 +3747,7 @@ bool UsdToUnreal::ConvertGeomPrimitiveTransform(const pxr::UsdPrim& InPrim, cons
 			Axis = UsdUtils::GetUsdValue<pxr::TfToken>(Attr, pxr::UsdTimeCode::Default());
 			if (Attr.ValueMightBeTimeVarying())
 			{
-				UE_LOG(
-					LogUsd,
-					Warning,
+				USD_LOG_WARNING(
 					TEXT("Animation of the 'axis' attribute for prim '%s' is not supported!"),
 					*UsdToUnreal::ConvertPath(InPrim.GetPrimPath())
 				);
@@ -3789,9 +3780,7 @@ bool UsdToUnreal::ConvertGeomPrimitiveTransform(const pxr::UsdPrim& InPrim, cons
 			Axis = UsdUtils::GetUsdValue<pxr::TfToken>(Attr, pxr::UsdTimeCode::Default());
 			if (Attr.ValueMightBeTimeVarying())
 			{
-				UE_LOG(
-					LogUsd,
-					Warning,
+				USD_LOG_WARNING(
 					TEXT("Animation of the 'axis' attribute for prim '%s' is not supported!"),
 					*UsdToUnreal::ConvertPath(InPrim.GetPrimPath())
 				);
@@ -3823,9 +3812,7 @@ bool UsdToUnreal::ConvertGeomPrimitiveTransform(const pxr::UsdPrim& InPrim, cons
 			Axis = UsdUtils::GetUsdValue<pxr::TfToken>(Attr, pxr::UsdTimeCode::Default());
 			if (Attr.ValueMightBeTimeVarying())
 			{
-				UE_LOG(
-					LogUsd,
-					Warning,
+				USD_LOG_WARNING(
 					TEXT("Animation of the 'axis' attribute for prim '%s' is not supported!"),
 					*UsdToUnreal::ConvertPath(InPrim.GetPrimPath())
 				);
@@ -3849,9 +3836,7 @@ bool UsdToUnreal::ConvertGeomPrimitiveTransform(const pxr::UsdPrim& InPrim, cons
 			Axis = UsdUtils::GetUsdValue<pxr::TfToken>(Attr, pxr::UsdTimeCode::Default());
 			if (Attr.ValueMightBeTimeVarying())
 			{
-				UE_LOG(
-					LogUsd,
-					Warning,
+				USD_LOG_WARNING(
 					TEXT("Animation of the 'axis' attribute for prim '%s' is not supported!"),
 					*UsdToUnreal::ConvertPath(InPrim.GetPrimPath())
 				);
@@ -4047,14 +4032,14 @@ UsdUtils::FUsdPrimMaterialAssignmentInfo UsdUtils::GetPrimMaterialAssignments(
 				return ValidPackagePath;
 			}
 
-			UE_LOG(
-				LogUsd,
-				Warning,
-				TEXT("Could not find a valid material at path '%s', targetted by prim '%s's unrealMaterial attribute. Material assignment will "
-					 "fallback to USD materials and display color data."),
-				*ValidPackagePath,
-				*UsdToUnreal::ConvertPath(UsdPrim.GetPath())
-			);
+			USD_LOG_USERWARNING(FText::Format(
+				LOCTEXT(
+					"UnrealMaterialNotFound",
+					"Could not find a valid material at path '{0}', targetted by prim '{1}'s unrealMaterial attribute. Material assignment will fallback to USD materials and display color data."
+				),
+				FText::FromString(ValidPackagePath),
+				FText::FromString(UsdToUnreal::ConvertPath(UsdPrim.GetPath()))
+			));
 		}
 
 		return {};
@@ -4234,7 +4219,7 @@ UsdUtils::FUsdPrimMaterialAssignmentInfo UsdUtils::GetPrimMaterialAssignments(
 			const bool bHasUnassignedIndices = SubsetIndexTotal != NumFaces;
 			if (bHasUnassignedIndices)
 			{
-				UE_LOG(LogUsd, Warning, TEXT("Found an invalid GeomSubsets partition in prim '%s'"), *UsdToUnreal::ConvertPath(UsdPrim.GetPath()));
+				USD_LOG_WARNING(TEXT("Found an invalid GeomSubsets partition in prim '%s'"), *UsdToUnreal::ConvertPath(UsdPrim.GetPath()));
 
 				// Assign these leftover indices to the *next* material slot we'll create (doesn't exist yet),
 				// which will be the "main" material assignment slot
@@ -4475,9 +4460,7 @@ bool UnrealToUsd::ConvertStaticMesh(
 	pxr::UsdVariantSet VariantSet = VariantSets.GetVariantSet(UnrealIdentifiers::LOD);
 	if (NumLODs > 1 && VariantSets.HasVariantSet(UnrealIdentifiers::LOD))
 	{
-		UE_LOG(
-			LogUsd,
-			Error,
+		USD_LOG_ERROR(
 			TEXT("Failed to export higher LODs for mesh '%s', as the target prim already has a variant set named '%s'!"),
 			*StaticMesh->GetName(),
 			*UsdToUnreal::ConvertToken(UnrealIdentifiers::LOD)
@@ -4645,7 +4628,7 @@ bool UnrealToUsd::ConvertStaticMesh(
 			const FStaticMeshLODResources* RenderMesh = &StaticMesh->GetLODForExport(LODIndex);
 			if (!RenderMesh || RenderMesh->VertexBuffers.StaticMeshVertexBuffer.GetNumVertices() == 0 || RenderMesh->Sections.Num() == 0)
 			{
-				UE_LOG(LogUsd, Warning, TEXT("Found invalid render data for LOD '%d' of '%s'!"), LODIndex, *StaticMesh->GetPathName());
+				USD_LOG_WARNING(TEXT("Found invalid render data for LOD '%d' of '%s'!"), LODIndex, *StaticMesh->GetPathName());
 				continue;
 			}
 
@@ -4655,9 +4638,7 @@ bool UnrealToUsd::ConvertStaticMesh(
 
 		if (!bExported)
 		{
-			UE_LOG(
-				LogUsd,
-				Warning,
+			USD_LOG_WARNING(
 				TEXT("Failed to export LOD '%d' of mesh '%s' onto prim '%s'!"),
 				LODIndex,
 				*StaticMesh->GetPathName(),
@@ -4702,9 +4683,7 @@ bool UnrealToUsd::ConvertMeshDescriptions(
 	pxr::UsdVariantSets VariantSets = UsdPrim.GetVariantSets();
 	if (NumLODs > 1 && VariantSets.HasVariantSet(UnrealIdentifiers::LOD))
 	{
-		UE_LOG(
-			LogUsd,
-			Error,
+		USD_LOG_ERROR(
 			TEXT("Failed to convert higher mesh description LODs for prim '%s', as the target prim already has a variant set named '%s'!"),
 			*UsdToUnreal::ConvertPath(UsdPrim.GetPath()),
 			*UsdToUnreal::ConvertToken(UnrealIdentifiers::LOD)
@@ -5891,31 +5870,8 @@ void UsdUtils::RepairNormalsAndTangents(const FString& PrimPath, FMeshDescriptio
 	if (InvalidNormalFraction >= GMeshNormalRepairThreshold)
 	{
 		Options |= EComputeNTBsFlags::Normals;
-		UE_LOG(
-			LogUsd,
-			Verbose,
-			TEXT("%f%% of the normals from Mesh prim '%s' are invalid or unusable. This is at or above the threshold of '%f%%' (configurable via "
-				 "the cvar '%s'), so normals will be discarded and fully recomputed. Note that when the cvar "
-				 "'USD.Subdiv.IgnoreNormalsWhenSubdividing' is true it is expected for subdivision meshes to have their normals discarded."),
-			InvalidNormalFraction * 100.0f,
-			*PrimPath,
-			GMeshNormalRepairThreshold * 100.0f,
-			*MeshNormalRepairThresholdText
-		);
-	}
-	else if (InvalidNormalFraction > 0)
-	{
-		UE_LOG(
-			LogUsd,
-			Verbose,
-			TEXT("%f%% of the normals from Mesh prim '%s' are invalid or unusable. This is below the threshold of '%f%%' (configurable via the "
-				 "cvar '%s'), so the invalid normals will be repaired."),
-			InvalidNormalFraction * 100.0f,
-			*PrimPath,
-			GMeshNormalRepairThreshold * 100.0f,
-			*MeshNormalRepairThresholdText
-		);
 	}
+
 	FStaticMeshOperations::ComputeTangentsAndNormals(MeshDescription, Options);
 }
 
@@ -6232,14 +6188,14 @@ bool UsdUtils::IterateLODMeshes(const pxr::UsdPrim& ParentPrim, TFunction<bool(c
 				}
 				else
 				{
-					UE_LOG(
-						LogUsd,
-						Warning,
-						TEXT("Unexpected prim '%s' inside LOD variant '%s'. For automatic parsing of LODs, each LOD variant should contain only a "
-							 "single Mesh prim named the same as the variant!"),
-						*UsdToUnreal::ConvertPath(Child.GetPath()),
-						*UsdToUnreal::ConvertString(LODVariantName)
-					);
+					USD_LOG_USERWARNING(FText::Format(
+						LOCTEXT(
+							"UnexpectedLODPrim",
+							"Unexpected prim '{0}' inside LOD variant '{1}'. For automatic parsing of LODs, each LOD variant should contain only a single Mesh prim named the same as the variant!"
+						),
+						FText::FromString(UsdToUnreal::ConvertPath(Child.GetPath())),
+						FText::FromString(UsdToUnreal::ConvertString(LODVariantName))
+					));
 				}
 			}
 		}
diff --git a/Engine/Plugins/Runtime/USDCore/Source/USDUtilities/Private/USDLayerUtils.cpp b/Engine/Plugins/Runtime/USDCore/Source/USDUtilities/Private/USDLayerUtils.cpp
--- a/Engine/Plugins/Runtime/USDCore/Source/USDUtilities/Private/USDLayerUtils.cpp
+++ b/Engine/Plugins/Runtime/USDCore/Source/USDUtilities/Private/USDLayerUtils.cpp
@@ -4,7 +4,6 @@
 
 #include "UnrealUSDWrapper.h"
 #include "USDErrorUtils.h"
-#include "USDLog.h"
 #include "USDMemory.h"
 #include "USDTypesConversion.h"
 
@@ -164,7 +163,9 @@ bool UsdUtils::InsertSubLayer(
 	FString ParentLayerPath = UsdToUnreal::ConvertString(ParentLayer->GetRealPath());
 	if (FPaths::GetCleanFilename(ParentLayerPath) == RelativeSubLayerPath)
 	{
-		UE_LOG(LogUsd, Warning, TEXT("Tried to add layer '%s' as a sublayer of itself!"), *ParentLayerPath);
+		USD_LOG_USERWARNING(
+			FText::Format(LOCTEXT("SelfSublayer", "Tried to add layer '{0}' as a sublayer of itself!"), FText::FromString(ParentLayerPath))
+		);
 		return false;
 	}
 
@@ -222,7 +223,7 @@ TOptional<FString> UsdUtils::BrowseUsdFile(EBrowseFileMode Mode)
 
 	if (SupportedExtensions.Num() == 0)
 	{
-		UE_LOG(LogUsd, Error, TEXT("No file extensions supported by the USD SDK!"));
+		USD_LOG_ERROR(TEXT("No file extensions supported by the USD SDK!"));
 		return {};
 	}
 
@@ -388,13 +389,11 @@ bool UsdUtils::DecompressUSDZFile(const FString& InUSDZFilePath, const FString&
 			}
 			else
 			{
-				UE_LOG(
-					LogUsd,
-					Warning,
-					TEXT("Failed to decompress %s: First file '%s' should always be an USD file!"),
-					*InUSDZFilePath,
-					*UEInnerFilePath
-				);
+				USD_LOG_USERWARNING(FText::Format(
+					LOCTEXT("FailedToDecompress", "Failed to decompress {0}: First file '{1}' should always be an USD file!"),
+					FText::FromString(InUSDZFilePath),
+					FText::FromString(UEInnerFilePath)
+				));
 				return false;
 			}
 		}
@@ -403,12 +402,20 @@ bool UsdUtils::DecompressUSDZFile(const FString& InUSDZFilePath, const FString&
 		// compression/encryption, and Usd_UsdzResolver::OpenAsset also just errors out if it runs into these
 		if (Info.compressionMethod != 0)
 		{
-			UE_LOG(LogUsd, Warning, TEXT("Cannot open %s in %s: Compressed files are not supported!"), *UEInnerFilePath, *InUSDZFilePath);
+			USD_LOG_USERWARNING(FText::Format(
+				LOCTEXT("FailedToDecompressInternal", "Cannot open {0} in {1}: Compressed files are not supported!"),
+				FText::FromString(UEInnerFilePath),
+				FText::FromString(InUSDZFilePath)
+			));
 			continue;
 		}
 		if (Info.encrypted)
 		{
-			UE_LOG(LogUsd, Warning, TEXT("Cannot open %s in %s: Encrypted files are not supported!"), *UEInnerFilePath, *InUSDZFilePath);
+			USD_LOG_USERWARNING(FText::Format(
+				LOCTEXT("FailToDecrypt", "Cannot open {0} in {1}: Encrypted files are not supported!"),
+				FText::FromString(UEInnerFilePath),
+				FText::FromString(InUSDZFilePath)
+			));
 			continue;
 		}
 
@@ -451,7 +458,7 @@ TUsdStore<pxr::SdfLayerRefPtr> UsdUtils::CreateNewLayer(
 	// New layer needs to be created and in the stage layer stack before we can edit it
 	UsdUtils::InsertSubLayer(ParentLayer.Get(), LayerFilePath);
 
-	UsdUtils::StartMonitoringErrors();
+	FScopedUsdMessageLog ScopedLog;
 	pxr::UsdEditContext UsdEditContext(UsdStage.Get(), LayerRef);
 
 	// Create default prim
@@ -466,9 +473,7 @@ TUsdStore<pxr::SdfLayerRefPtr> UsdUtils::CreateNewLayer(
 		LayerRef->SetDefaultPrim(DefaultPrim.GetPrim().GetName());
 	}
 
-	bool bHadErrors = UsdUtils::ShowErrorsAndStopMonitoring();
-
-	if (bHadErrors)
+	if (FUsdLogManager::HasAccumulatedErrors())
 	{
 		return {};
 	}
@@ -711,7 +716,7 @@ void UsdUtils::AddTimeCodeRangeToLayer(const pxr::SdfLayerRefPtr& Layer, double
 
 	if (!Layer)
 	{
-		FUsdLogManager::LogMessage(EMessageSeverity::Warning, LOCTEXT("AddTimeCode_InvalidLayer", "Trying to set timecodes on an invalid layer."));
+		USD_LOG_USERWARNING(LOCTEXT("AddTimeCode_InvalidLayer", "Trying to set timecodes on an invalid layer."));
 		return;
 	}
 
diff --git a/Engine/Plugins/Runtime/USDCore/Source/USDUtilities/Private/USDLightConversion.cpp b/Engine/Plugins/Runtime/USDCore/Source/USDUtilities/Private/USDLightConversion.cpp
--- a/Engine/Plugins/Runtime/USDCore/Source/USDUtilities/Private/USDLightConversion.cpp
+++ b/Engine/Plugins/Runtime/USDCore/Source/USDUtilities/Private/USDLightConversion.cpp
@@ -7,8 +7,8 @@
 #include "USDAttributeUtils.h"
 #include "USDClassesModule.h"
 #include "USDConversionUtils.h"
+#include "USDErrorUtils.h"
 #include "USDLayerUtils.h"
-#include "USDLog.h"
 #include "USDMemory.h"
 #include "USDObjectUtils.h"
 #include "USDShadeConversion.h"
@@ -41,6 +41,8 @@
 #include "pxr/usd/usdLux/sphereLight.h"
 #include "USDIncludesEnd.h"
 
+#define LOCTEXT_NAMESPACE "USDLightConversion"
+
 namespace LightConversionImpl
 {
 	/**
@@ -243,13 +245,11 @@ bool UsdToUnreal::ConvertDomeLight(
 		// Show a good warning for this because it's easy to pick some cubemap asset from the engine (that usually don't come with the
 		// source texture) and have the dome light silently not work again
 		FString TargetAssetPath = UsdToUnreal::ConvertString(TextureAssetPath.GetAssetPath());
-		UE_LOG(
-			LogUsd,
-			Warning,
-			TEXT("Failed to find texture '%s' used for UsdLuxDomeLight '%s'!"),
-			*TargetAssetPath,
-			*UsdToUnreal::ConvertPath(DomeLight.GetPrim().GetPath())
-		);
+		USD_LOG_USERWARNING(FText::Format(
+			LOCTEXT("MissingDomeLightTexture", "Failed to find texture '{0}' used for UsdLuxDomeLight '{1}'!"),
+			FText::FromString(TargetAssetPath),
+			FText::FromString(UsdToUnreal::ConvertPath(DomeLight.GetPrim().GetPath()))
+		));
 
 		return true;
 	}
@@ -655,15 +655,15 @@ bool UnrealToUsd::ConvertSkyLightComponent(const USkyLightComponent& LightCompon
 				FString FilePath = AssetImportData->GetFirstFilename();
 				if (!FPaths::FileExists(FilePath))
 				{
-					UE_LOG(
-						LogUsd,
-						Warning,
-						TEXT("Used '%s' as cubemap when converting SkyLightComponent '%s' onto prim '%s', but the cubemap does not exist on the "
-							 "filesystem!"),
-						*FilePath,
-						*LightComponent.GetPathName(),
-						*UsdToUnreal::ConvertPath(Prim.GetPrimPath())
-					);
+					USD_LOG_USERWARNING(FText::Format(
+						LOCTEXT(
+							"SourceCubemapDoesntExist",
+							"Used '{0}' as cubemap when converting SkyLightComponent '{1}' onto prim '{2}', but the cubemap does not exist on the filesystem!"
+						),
+						FText::FromString(FilePath),
+						FText::FromString(LightComponent.GetPathName()),
+						FText::FromString(UsdToUnreal::ConvertPath(Prim.GetPrimPath()))
+					));
 				}
 
 				UsdUtils::MakePathRelativeToLayer(UE::FSdfLayer{Prim.GetStage()->GetEditTarget().GetLayer()}, FilePath);
@@ -837,4 +837,6 @@ float UsdUtils::ConvertIntensityToNits(float Intensity, float Steradians, float
 	return Intensity;
 }
 
+#undef LOCTEXT_NAMESPACE
+
 #endif	  // #if USE_USD_SDK
diff --git a/Engine/Plugins/Runtime/USDCore/Source/USDUtilities/Private/USDMaterialXShaderGraph.cpp b/Engine/Plugins/Runtime/USDCore/Source/USDUtilities/Private/USDMaterialXShaderGraph.cpp
--- a/Engine/Plugins/Runtime/USDCore/Source/USDUtilities/Private/USDMaterialXShaderGraph.cpp
+++ b/Engine/Plugins/Runtime/USDCore/Source/USDUtilities/Private/USDMaterialXShaderGraph.cpp
@@ -4,12 +4,12 @@
 
 #ifdef ENABLE_USD_MATERIALX
 
+#include "Misc/Paths.h"
+#include "USDErrorUtils.h"
 #include "USDLayerUtils.h"
-#include "USDLog.h"
 #include "USDMemory.h"
 #include "USDTypesConversion.h"
 #include "UsdWrappers/SdfLayer.h"
-#include "Misc/Paths.h"
 
 namespace mx = MaterialX;
 
@@ -65,20 +65,18 @@ FUsdMaterialXShaderGraph::FUsdMaterialXShaderGraph(const pxr::UsdPrim& UsdShadeM
 		"geometry_coat_tangent"
 	}
 {
-	pxr::UsdShadeMaterial UsdShadeMaterial{ UsdShadeMaterialPrim };
+	pxr::UsdShadeMaterial UsdShadeMaterial{UsdShadeMaterialPrim};
 
 	if (!UsdShadeMaterial)
 	{
-		UE_LOG(LogUsd, Error, TEXT("Couldn't create a UsdShadeMaterial for the prim."));
+		USD_LOG_ERROR(TEXT("Couldn't create a UsdShadeMaterial for the prim."));
 		return;
 	}
 
 	Document = mx::createDocument();
 	mx::DocumentPtr MaterialXLibrary = mx::createDocument();
-	mx::FileSearchPath MaterialXFolder{
-		TCHAR_TO_UTF8(*FPaths::Combine(FPaths::EngineDir(), TEXT("Binaries"), TEXT("ThirdParty"), TEXT("MaterialX")))
-	};
-	mx::StringSet LoadedLibs = mx::loadLibraries({ "libraries" }, MaterialXFolder, MaterialXLibrary);
+	mx::FileSearchPath MaterialXFolder{TCHAR_TO_UTF8(*FPaths::Combine(FPaths::EngineDir(), TEXT("Binaries"), TEXT("ThirdParty"), TEXT("MaterialX")))};
+	mx::StringSet LoadedLibs = mx::loadLibraries({"libraries"}, MaterialXFolder, MaterialXLibrary);
 	Document->importLibrary(MaterialXLibrary);
 
 	CreateSurfaceMaterial(UsdShadeMaterial, RenderContext);
@@ -108,7 +106,10 @@ void FUsdMaterialXShaderGraph::CreateSurfaceMaterial(const pxr::UsdShadeMaterial
 
 	if (!SurfaceShader)
 	{
-		UE_LOG(LogUsd, Error, TEXT("Couldn't find a SurfaceShader for UsdShadeMaterial '%s'."), UTF8_TO_TCHAR(UsdShadeMaterial.GetPrim().GetName().GetString().c_str()));
+		USD_LOG_ERROR(
+			TEXT("Couldn't find a SurfaceShader for UsdShadeMaterial '%s'."),
+			*UsdToUnreal::ConvertToken(UsdShadeMaterial.GetPrim().GetName())
+		);
 		return;
 	}
 
@@ -170,7 +171,7 @@ void FUsdMaterialXShaderGraph::CreateSurfaceMaterial(const pxr::UsdShadeMaterial
 		SurfaceMaterial->validate(&ErrorMessage);
 		if (!ErrorMessage.empty())
 		{
-			UE_LOG(LogUsd, Error, TEXT("USD MaterialX: %s."), UTF8_TO_TCHAR(ErrorMessage.c_str()));
+			USD_LOG_ERROR(TEXT("USD MaterialX: %s."), UTF8_TO_TCHAR(ErrorMessage.c_str()));
 			SurfaceMaterial = nullptr;
 		}
 	}
@@ -187,10 +188,10 @@ bool FUsdMaterialXShaderGraph::ComputeShaderGraph(mx::NodePtr Node, const pxr::U
 
 	for (const pxr::UsdShadeInput& UsdInput : *UsdInputs)
 	{
-		// Basically the idea here is to traverse the entire graph coming from a tangent space input of a surface shader (e.g: 'normal' of <standard_surface>)
-		// if we come from a TS input, all geompropvalue along the path needs to be set in TS
-		// since we can only set at init-time the value of a TGuardValue and we're in a recursive state,
-		// if bTangentSpaceInput is false we set its value to IsTangentSpace otherwise we keep it as is (meaning true)
+		// Basically the idea here is to traverse the entire graph coming from a tangent space input of a surface shader (e.g: 'normal' of
+		// <standard_surface>) if we come from a TS input, all geompropvalue along the path needs to be set in TS since we can only set at init-time
+		// the value of a TGuardValue and we're in a recursive state, if bTangentSpaceInput is false we set its value to IsTangentSpace otherwise we
+		// keep it as is (meaning true)
 		TGuardValue<bool> GuardInput(bTangentSpaceInput, bTangentSpaceInput || IsTangentSpaceInput(UsdInput));
 
 		mx::InputPtr Input = GetInput(Node, UsdInput);
@@ -199,11 +200,11 @@ bool FUsdMaterialXShaderGraph::ComputeShaderGraph(mx::NodePtr Node, const pxr::U
 		TUsdStore<pxr::UsdShadeInput::SourceInfoVector> ConnectedSources = UsdInput.GetConnectedSources();
 
 		// At this point an input should end up in either of these 3 cases, with the precedence in that order:
-		// 1. it has an interfaceName, if that's the case we continue to the next input (it has no value, no nodename)		
+		// 1. it has an interfaceName, if that's the case we continue to the next input (it has no value, no nodename)
 		// 2. it is connected to a nodename with a valid nodedef (it has no value)
 		// 3. it has a value, if not defined by the input, the default one from the nodedef is taken
-		
-		//Let's loop over the connected sources, even though we should expect only one
+
+		// Let's loop over the connected sources, even though we should expect only one
 		for (const pxr::UsdShadeConnectionSourceInfo& Source : *ConnectedSources)
 		{
 			if (SetInterfaceName(Input, Source))
@@ -234,7 +235,7 @@ bool FUsdMaterialXShaderGraph::ComputeShaderGraph(mx::NodePtr Node, const pxr::U
 			}
 
 			// recurse if it's a newly created node, otherwise the work has already been done
-			if(!bIsNodeAlreadyExist)
+			if (!bIsNodeAlreadyExist)
 			{
 				bHasValue = ComputeShaderGraph(ConnectedNode, Source.source);
 			}
@@ -253,7 +254,7 @@ bool FUsdMaterialXShaderGraph::ComputeShaderGraph(mx::NodePtr Node, const pxr::U
 
 bool FUsdMaterialXShaderGraph::IsGeomColor(const pxr::UsdShadeShader& GeomPropValueShader)
 {
-	if (pxr::UsdShadeInput GeomPropInput = GeomPropValueShader.GetInput(pxr::TfToken{ "geomprop" }))
+	if (pxr::UsdShadeInput GeomPropInput = GeomPropValueShader.GetInput(pxr::TfToken{"geomprop"}))
 	{
 		if (std::string Value; GeomPropInput.Get<std::string>(&Value))
 		{
@@ -276,16 +277,17 @@ void FUsdMaterialXShaderGraph::SetInputValue(mx::InputPtr Input, const pxr::UsdS
 
 		if (FPairingUsdMtlxType Value; UsdInput.Get<FPairingUsdMtlxType>(&Value))
 		{
-			if constexpr (std::is_same_v<FPairingUsdMtlxType, pxr::GfVec4f> || std::is_same_v<FPairingUsdMtlxType, pxr::GfVec3f> || std::is_same_v<FPairingUsdMtlxType, pxr::GfVec2f>)
+			if constexpr (std::is_same_v<FPairingUsdMtlxType, pxr::GfVec4f> || std::is_same_v<FPairingUsdMtlxType, pxr::GfVec3f>
+						  || std::is_same_v<FPairingUsdMtlxType, pxr::GfVec2f>)
 			{
-				Input->setValue(MtlxType{ Value.data(), Value.data() + FPairingUsdMtlxType::dimension });
+				Input->setValue(MtlxType{Value.data(), Value.data() + FPairingUsdMtlxType::dimension});
 			}
 			else
 			{
 				Input->setValue(Value);
 			}
 		}
-		else if(mx::NodePtr Node = Input->getParent()->asA<mx::Node>())
+		else if (mx::NodePtr Node = Input->getParent()->asA<mx::Node>())
 		{
 			mx::NodeDefPtr NodeDef = Node->getNodeDef();
 			mx::InputPtr ActiveInput = NodeDef->getActiveInput(Input->getName());
@@ -329,7 +331,7 @@ void FUsdMaterialXShaderGraph::SetInputValue(mx::InputPtr Input, const pxr::UsdS
 		{
 			Input->setValue(Value.Get());
 		}
-		else if(mx::NodePtr Node = Input->getParent()->asA<mx::Node>())
+		else if (mx::NodePtr Node = Input->getParent()->asA<mx::Node>())
 		{
 			mx::NodeDefPtr NodeDef = Node->getNodeDef();
 			mx::InputPtr ActiveInput = NodeDef->getActiveInput(Input->getName());
@@ -340,26 +342,28 @@ void FUsdMaterialXShaderGraph::SetInputValue(mx::InputPtr Input, const pxr::UsdS
 	else if (InputType == "filename")
 	{
 		FString PathToResolve;
-		if(pxr::SdfAssetPath AssetPath; UsdInput.Get<pxr::SdfAssetPath>(&AssetPath))
+		if (pxr::SdfAssetPath AssetPath; UsdInput.Get<pxr::SdfAssetPath>(&AssetPath))
 		{
 			TUsdStore<UE::FSdfLayer> Layer = UsdUtils::FindLayerForAttribute(UsdInput.GetAttr(), pxr::UsdTimeCode::Default().GetValue());
 
 			PathToResolve = Layer.Get()
-				? UsdToUnreal::ConvertString(TUsdStore<std::string>(pxr::SdfComputeAssetPathRelativeToLayer(Layer.Get(), AssetPath.GetAssetPath())).Get())
-				: UsdToUnreal::ConvertString(AssetPath.GetAssetPath());
+								? UsdToUnreal::ConvertString(
+									  TUsdStore<std::string>(pxr::SdfComputeAssetPathRelativeToLayer(Layer.Get(), AssetPath.GetAssetPath())).Get()
+								  )
+								: UsdToUnreal::ConvertString(AssetPath.GetAssetPath());
 		}
-		Input->setValueString(std::string{ TCHAR_TO_UTF8(*PathToResolve) });
+		Input->setValueString(std::string{TCHAR_TO_UTF8(*PathToResolve)});
 	}
 	else
 	{
-		UE_LOG(LogUsd, Warning, TEXT("Couldn't find a value type for (%s)."), UTF8_TO_TCHAR(InputType.c_str()));
+		USD_LOG_WARNING(TEXT("Couldn't find a value type for (%s)."), UTF8_TO_TCHAR(InputType.c_str()));
 	}
 }
 
 bool FUsdMaterialXShaderGraph::SetInterfaceName(MaterialX::InputPtr Input, const pxr::UsdShadeConnectionSourceInfo& Source)
 {
 	bool bIsSet = false;
-	const std::string& SourceName= Source.sourceName.GetString();
+	const std::string& SourceName = Source.sourceName.GetString();
 	if (InterfaceNames.find(SourceName) != InterfaceNames.end())
 	{
 		if (mx::InputPtr InputInterfaceName = Document->getInput(SourceName))
@@ -385,12 +389,12 @@ mx::NodePtr FUsdMaterialXShaderGraph::GetNode(const pxr::UsdShadeConnectionSourc
 	std::string ConnectedShaderName = Source.source.GetPrim().GetName().GetString();
 
 	mx::NodePtr ConnectedNode = Document->getNode(ConnectedShaderName);
-	
+
 	// if the node is not present in the graph, we create it and notify that it's a new node
 	if (!ConnectedNode)
 	{
 		ConnectedNode = Document->addNodeInstance(NodeDef, ConnectedShaderName);
-		if(bIsGeomPropValue)
+		if (bIsGeomPropValue)
 		{
 			ConnectedNode->setTypedAttribute("UE:GeomPropImage", bIsGeomPropValue);
 		}
@@ -420,10 +424,12 @@ mx::InputPtr FUsdMaterialXShaderGraph::GetInput(mx::NodePtr Node, const pxr::Usd
 			InputType = "filename";
 			TUsdStore<std::string> Value;
 			UsdInput.Get<std::string>(&Value.Get());
-			if (!GeomPropValueNames.FindByPredicate([&Value](const FGeomProp& Other)
-				{
-					return Other.Name == Value->c_str();
-				}))
+			if (!GeomPropValueNames.FindByPredicate(
+					[&Value](const FGeomProp& Other)
+					{
+						return Other.Name == Value->c_str();
+					}
+				))
 			{
 				// we only set the the input as being in Tangent Space if also its type is a vec3
 				GeomPropValueNames.Emplace(Value->c_str(), bTangentSpaceInput && Node->getType() == "vector3" ? bTangentSpaceInput : false);
@@ -436,11 +442,11 @@ mx::InputPtr FUsdMaterialXShaderGraph::GetInput(mx::NodePtr Node, const pxr::Usd
 
 mx::NodeDefPtr FUsdMaterialXShaderGraph::GetNodeDef(const pxr::UsdShadeConnectionSourceInfo& Source, bool& bIsGeomPropValue)
 {
-	pxr::UsdShadeShader UsdConnectedShader{ Source.source.GetPrim() };
+	pxr::UsdShadeShader UsdConnectedShader{Source.source.GetPrim()};
 
 	if (!UsdConnectedShader)
-	{		                                                                                      
-		UE_LOG(LogUsd, Warning, TEXT("The '%s' connected source is not a valid USD shader."), *UsdToUnreal::ConvertString(Source.source.GetPrim().GetName()));
+	{
+		USD_LOG_WARNING(TEXT("The '%s' connected source is not a valid USD shader."), *UsdToUnreal::ConvertString(Source.source.GetPrim().GetName()));
 	}
 
 	pxr::TfToken NodeDefTokenId;
@@ -475,9 +481,11 @@ mx::NodeDefPtr FUsdMaterialXShaderGraph::GetNodeDef(const pxr::UsdShadeConnectio
 			}
 			else if (pos = NodeDefString.rfind("string"); pos != std::string::npos)
 			{
-				UE_LOG(LogUsd, Warning, TEXT("'$s': '%s' are not supported"),
-					   *UsdToUnreal::ConvertString(Source.source.GetPrim().GetName()),
-					   *UsdToUnreal::ConvertString(NodeDefString));
+				USD_LOG_WARNING(
+					TEXT("'$s': '%s' are not supported"),
+					*UsdToUnreal::ConvertString(Source.source.GetPrim().GetName()),
+					*UsdToUnreal::ConvertString(NodeDefString)
+				);
 				return nullptr;
 			}
 		}
@@ -485,10 +493,10 @@ mx::NodeDefPtr FUsdMaterialXShaderGraph::GetNodeDef(const pxr::UsdShadeConnectio
 	}
 
 	mx::NodeDefPtr NodeDef = Document->getNodeDef(NodeDefString);
-	
-	if(!NodeDef)
+
+	if (!NodeDef)
 	{
-		UE_LOG(LogUsd, Warning, TEXT("Couldn't find a nodedef for (%s)."), UTF8_TO_TCHAR(NodeDefTokenId.GetString().c_str()));
+		USD_LOG_WARNING(TEXT("Couldn't find a nodedef for (%s)."), UTF8_TO_TCHAR(NodeDefTokenId.GetString().c_str()));
 	}
 
 	return NodeDef;
@@ -505,7 +513,7 @@ bool FUsdMaterialXShaderGraph::GetMatchingInputType(const std::string& InputName
 			InOutInputType = It->second;
 			bMatch = true;
 		}
-		else if (InputName == "file") // we're dealing with a texture, USD calls the type "asset" for the input
+		else if (InputName == "file")	 // we're dealing with a texture, USD calls the type "asset" for the input
 		{
 			InOutInputType = "filename";
 			bMatch = true;
@@ -519,4 +527,4 @@ bool FUsdMaterialXShaderGraph::IsTangentSpaceInput(const pxr::UsdShadeInput& Usd
 {
 	return TangentSpaceInputs.find(UsdInput.GetBaseName().GetString()) != TangentSpaceInputs.cend();
 }
-#endif // ENABLE_USD_MATERIALX
\ No newline at end of file
+#endif	  // ENABLE_USD_MATERIALX
diff --git a/Engine/Plugins/Runtime/USDCore/Source/USDUtilities/Private/USDPrimConversion.cpp b/Engine/Plugins/Runtime/USDCore/Source/USDUtilities/Private/USDPrimConversion.cpp
--- a/Engine/Plugins/Runtime/USDCore/Source/USDUtilities/Private/USDPrimConversion.cpp
+++ b/Engine/Plugins/Runtime/USDCore/Source/USDUtilities/Private/USDPrimConversion.cpp
@@ -7,9 +7,9 @@
 #include "USDAttributeUtils.h"
 #include "USDConversionUtils.h"
 #include "USDDrawModeComponent.h"
+#include "USDErrorUtils.h"
 #include "USDLayerUtils.h"
 #include "USDLightConversion.h"
-#include "USDLog.h"
 #include "USDMemory.h"
 #include "USDObjectUtils.h"
 #include "USDShadeConversion.h"
@@ -124,6 +124,8 @@
 #include "pxr/usd/usdVol/volume.h"
 #include "USDIncludesEnd.h"
 
+#define LOCTEXT_NAMESPACE "USDPrimConversion"
+
 static bool GConsiderAllPrimsHaveAnimatedBounds = false;
 static FAutoConsoleVariableRef CVarConsiderAllPrimsHaveAnimatedBounds(
 	TEXT("USD.Bounds.ConsiderAllPrimsHaveAnimatedBounds"),
@@ -287,15 +289,16 @@ namespace UE::USDPrimConversion::Private
 			const FVector& Scale = ParentToWorld.GetScale3D();
 			if (!FMath::IsNearlyEqual(Scale.X, Scale.Y) || !FMath::IsNearlyEqual(Scale.X, Scale.Z))
 			{
-				UE_LOG(
-					LogUsd,
-					Warning,
-					TEXT("Inverting transform with non-uniform scaling '%s' when computing relative transform for prim '%s'! Result will "
-						 "likely be incorrect, since FTransforms can't invert non-uniform scalings. You can work around this by baking your "
-						 "non-uniform scaling transform into the vertices, or by not using the !resetXformStack! Xform op."),
-					*Scale.ToString(),
-					*UsdToUnreal::ConvertPath(UsdPrim.GetPrimPath())
-				);
+				USD_LOG_USERWARNING(FText::Format(
+					LOCTEXT(
+						"InvertNonUniformScaling",
+						"Inverting transform with non-uniform scaling '{0}' when computing relative transform for prim '{1}'! Result will "
+						"likely be incorrect, since FTransforms can't invert non-uniform scalings. You can work around this by baking your "
+						"non-uniform scaling transform into the vertices, or by not using the !resetXformStack! Xform op."
+					),
+					FText::FromString(Scale.ToString()),
+					FText::FromString(UsdToUnreal::ConvertPath(UsdPrim.GetPrimPath()))
+				));
 			}
 
 			// Multiplying with matrices here helps mitigate the issues encountered with non-uniform scaling, however it will stil
@@ -2396,9 +2399,7 @@ bool UnrealToUsd::ConvertAudioComponent(const UAudioComponent& AudioComponent, p
 		{
 			if (!FPaths::FileExists(FilePath))
 			{
-				UE_LOG(
-					LogUsd,
-					Warning,
+				USD_LOG_WARNING(
 					TEXT(
 						"Used '%s' as the audio file when converting AudioComponent '%s' onto prim '%s', but the file does not exist on the filesystem!"
 					),
@@ -3212,9 +3213,7 @@ bool UnrealToUsd::Convert3DTransformTrack(
 			TArrayView<FMovieSceneDoubleChannel*> Channels = TransformSection->GetChannelProxy().GetChannels<FMovieSceneDoubleChannel>();
 			if (Channels.Num() < 9)
 			{
-				UE_LOG(
-					LogUsd,
-					Error,
+				USD_LOG_ERROR(
 					TEXT("Unexpected number of double tracks (%d) in transform section '%s'"),
 					Channels.Num(),
 					*TransformSection->GetPathName()
@@ -3841,9 +3840,7 @@ bool UnrealToUsd::ConvertMaterialOverrides(
 
 		if (!UsdPrim.IsA<pxr::UsdSkelSkeleton>())
 		{
-			UE_LOG(
-				LogUsd,
-				Warning,
+			USD_LOG_WARNING(
 				TEXT("For the skeletal case, ConvertMaterialOverrides must now receive a Skeleton prim! ('%s' was provided)"),
 				*UsdToUnreal::ConvertPath(UsdPrim.GetPrimPath())
 			);
@@ -5174,12 +5171,7 @@ bool UnrealToUsd::CreateSkeletalAnimationBaker(
 	pxr::UsdSkelSkeleton UsdSkeleton{SkeletonPrim};
 	if (!UsdSkeleton)
 	{
-		UE_LOG(
-			LogUsd,
-			Warning,
-			TEXT("Failed to create skeletal animation baker: Prim '%s' must be a UsdSkeleton!"),
-			*SkeletonPrim.GetPrimPath().GetString()
-		);
+		USD_LOG_WARNING(TEXT("Failed to create skeletal animation baker: Prim '%s' must be a UsdSkeleton!"), *SkeletonPrim.GetPrimPath().GetString());
 		return false;
 	}
 
@@ -5393,9 +5385,7 @@ bool UnrealToUsd::CreateSkeletalAnimationBaker(
 			}
 			if (LocalBoneTransforms.Num() != NumBones)
 			{
-				UE_LOG(
-					LogUsd,
-					Warning,
+				USD_LOG_WARNING(
 					TEXT(
 						"Failed to retrieve bone transforms when baking skeletal animation for component '%s' at timeCode '%f'. Expected %d transforms, "
 						"received %d"
@@ -6835,17 +6825,18 @@ bool UnrealToUsd::ConvertDrawModeComponent(
 						{
 							if (!FPaths::FileExists(TextureSourcePath))
 							{
-								UE_LOG(
-									LogUsd,
-									Warning,
-									TEXT("Authoring card texture path '%s' for texture '%s' onto attribute '%s', but the source image file doesn't "
-										 "exist on disk! It may not be possible to display this card texture if the stage is reloaded or reopened. "
-										 "UTexture assets can't be automatically exported in this manner just yet. If you want to assign a new "
-										 "texture, make sure the UTexture asset's source image file exists on disk."),
-									*TextureSourcePath,
-									*FaceTexture->GetPathName(),
-									*UsdToUnreal::ConvertPath(Attr.GetPath())
-								);
+								USD_LOG_USERWARNING(FText::Format(
+									LOCTEXT(
+										"CardTextureNotOnDisk",
+										"Authoring card texture path '{0}' for texture '{1}' onto attribute '{2}', but the source image file doesn't "
+										"exist on disk! It may not be possible to display this card texture if the stage is reloaded or reopened. "
+										"UTexture assets can't be automatically exported in this manner just yet. If you want to assign a new "
+										"texture, make sure the UTexture asset's source image file exists on disk."
+									),
+									FText::FromString(TextureSourcePath),
+									FText::FromString(FaceTexture->GetPathName()),
+									FText::FromString(UsdToUnreal::ConvertPath(Attr.GetPath()))
+								));
 							}
 
 							pxr::SdfAssetPath AssetPath = pxr::SdfAssetPath{UnrealToUsd::ConvertString(*TextureSourcePath).Get()};
@@ -6854,15 +6845,16 @@ bool UnrealToUsd::ConvertDrawModeComponent(
 					}
 					else
 					{
-						UE_LOG(
-							LogUsd,
-							Warning,
-							TEXT("Not authoring card texture for attribute '%s' because the assigned texture '%s' has no AssetImportData! UTexture "
-								 "assets can't be automatically exported in this manner just yet. If you want to assign a new texture, make sure the "
-								 "UTexture asset's source image file exists on disk."),
-							*UsdToUnreal::ConvertPath(Attr.GetPath()),
-							*FaceTexture->GetPathName()
-						);
+						USD_LOG_USERWARNING(FText::Format(
+							LOCTEXT(
+								"NoAssetImportData",
+								"Not authoring card texture for attribute '{0}' because the assigned texture '{1}' has no AssetImportData! UTexture "
+								"assets can't be automatically exported in this manner just yet. If you want to assign a new texture, make sure the "
+								"UTexture asset's source image file exists on disk."
+							),
+							FText::FromString(UsdToUnreal::ConvertPath(Attr.GetPath())),
+							FText::FromString(FaceTexture->GetPathName())
+						));
 					}
 #endif	  // WITH_EDITOR
 				}
@@ -7021,9 +7013,7 @@ namespace UE::USDPrimConversion::Private
 
 			if (!bSuccess)
 			{
-				UE_LOG(
-					LogUsd,
-					Warning,
+				USD_LOG_WARNING(
 					TEXT("Failed to set metadata '%s' on prim '%s' as the value '%s' could not be parsed from string!"),
 					*FullKeyPath,
 					*UsdToUnreal::ConvertPath(Prim.GetPrimPath()),
@@ -7055,13 +7045,7 @@ namespace UE::USDPrimConversion::Private
 			const bool bWillOverwrite = Prim.HasMetadataDictKey(TopLevelKeyName, KeyPath);
 			if (bWillOverwrite)
 			{
-				UE_LOG(
-					LogUsd,
-					Log,
-					TEXT("Overwriting metadata field '%s' on prim '%s'"),
-					*FullKeyPath,
-					*UsdToUnreal::ConvertPath(Prim.GetPrimPath())
-				);
+				USD_LOG_INFO(TEXT("Overwriting metadata field '%s' on prim '%s'"), *FullKeyPath, *UsdToUnreal::ConvertPath(Prim.GetPrimPath()));
 			}
 
 			// This will also nicely create the nested dictionaries that it needs on-demand
@@ -7203,4 +7187,6 @@ bool UnrealToUsd::ConvertMetadata(
 	return true;
 }
 
+#undef LOCTEXT_NAMESPACE
+
 #endif	  // #if USE_USD_SDK
diff --git a/Engine/Plugins/Runtime/USDCore/Source/USDUtilities/Private/USDShadeConversion.cpp b/Engine/Plugins/Runtime/USDCore/Source/USDUtilities/Private/USDShadeConversion.cpp
--- a/Engine/Plugins/Runtime/USDCore/Source/USDUtilities/Private/USDShadeConversion.cpp
+++ b/Engine/Plugins/Runtime/USDCore/Source/USDUtilities/Private/USDShadeConversion.cpp
@@ -14,7 +14,6 @@
 #include "USDErrorUtils.h"
 #include "USDExportUtils.h"
 #include "USDLayerUtils.h"
-#include "USDLog.h"
 #include "USDMemory.h"
 #include "USDObjectUtils.h"
 #include "USDTypesConversion.h"
@@ -322,7 +321,7 @@ namespace UE::UsdShadeConversion::Private
 				}
 				else
 				{
-					UE_LOG(LogUsd, Warning, TEXT("Error creating texture. Bit depth is unsupported. (%d)"), BitDepth);
+					USD_LOG_WARNING(TEXT("Error creating texture. Bit depth is unsupported. (%d)"), BitDepth);
 					return nullptr;
 				}
 
@@ -346,7 +345,7 @@ namespace UE::UsdShadeConversion::Private
 		}
 		else
 		{
-			UE_LOG(LogUsd, Warning, TEXT("Error creating texture. Couldn't determine the file format"));
+			USD_LOG_WARNING(TEXT("Error creating texture. Couldn't determine the file format"));
 		}
 
 		return NewTexture;
@@ -482,9 +481,7 @@ namespace UE::UsdShadeConversion::Private
 			}
 			else
 			{
-				UE_LOG(
-					LogUsd,
-					Warning,
+				USD_LOG_WARNING(
 					TEXT("Unexpected shader node id '%s' when traversing texture node '%s' for primvars!"),
 					*UsdToUnreal::ConvertToken(ConnectedNodeId),
 					*UsdToUnreal::ConvertPath(UsdUVTexture.GetPrim().GetPrimPath())
@@ -876,9 +873,7 @@ namespace UE::UsdShadeConversion::Private
 					{
 						if (!FMath::IsNearlyEqual(UsdFallbackVec4[3], 1.0f))
 						{
-							UE_LOG(
-								LogUsd,
-								Warning,
+							USD_LOG_WARNING(
 								TEXT("Ignoring alpha value from fallback GfVec4f [%f, %f, %f, %f] used for Shader '%s'"),
 								UsdFallbackVec4[0],
 								UsdFallbackVec4[1],
@@ -1122,19 +1117,21 @@ namespace UE::UsdShadeConversion::Private
 			}
 			else
 			{
-				UE_LOG(
-					LogUsd,
-					Warning,
-					TEXT("Failed to find primvar '%s' when setting material parameter '%s' on material '%s'. Available primvars and UV "
-						 "indices: %s.%s"),
-					*TextureValue.Primvar,
-					ParameterName,
-					*Material.GetPathName(),
-					*UsdUtils::StringifyMap(PrimvarToUVIndex),
-					TextureValue.Primvar.IsEmpty() ? TEXT(" Is your UsdUVTexture Shader missing the 'inputs:st' attribute? (It specifies "
-														  "which UV set to sample the texture with)")
-												   : TEXT("")
-				);
+				USD_LOG_USERWARNING(FText::Format(
+					LOCTEXT(
+						"MissingPrimvar",
+						"Failed to find primvar '{0}' when setting material parameter '{1}' on material '{2}'. Available primvars and UV indices: {3}.{4}"
+					),
+					FText::FromString(TextureValue.Primvar),
+					FText::FromString(ParameterName),
+					FText::FromString(Material.GetPathName()),
+					FText::FromString(UsdUtils::StringifyMap(PrimvarToUVIndex)),
+					FText::FromString(
+						TextureValue.Primvar.IsEmpty() ? TEXT(" Is your UsdUVTexture Shader missing the 'inputs:st' attribute? (It specifies "
+															  "which UV set to sample the texture with)")
+													   : TEXT("")
+					)
+				));
 			}
 
 			FLinearColor ComponentMask = FLinearColor::Black;
@@ -1469,7 +1466,7 @@ namespace UE::UsdShadeConversion::Private
 			}
 			else
 			{
-				UE_LOG(LogUsd, Warning, TEXT("Failed to write texture '%s', baked channel will be ignored."), *TextureFilePath);
+				USD_LOG_WARNING(TEXT("Failed to write texture '%s', baked channel will be ignored."), *TextureFilePath);
 			}
 		}
 
@@ -1525,7 +1522,7 @@ namespace UE::UsdShadeConversion::Private
 			}
 			else
 			{
-				UE_LOG(LogUsd, Warning, TEXT("Failed to write texture '%s', baked channel will be ignored."), *TextureFilePath);
+				USD_LOG_WARNING(TEXT("Failed to write texture '%s', baked channel will be ignored."), *TextureFilePath);
 			}
 		}
 
@@ -1676,7 +1673,7 @@ namespace UE::UsdShadeConversion::Private
 
 			if ((NumSamples == 0 || !SampleSize) && !UserConstantValue)
 			{
-				UE_LOG(LogUsd, Log, TEXT("Skipping material property %d as we have no valid data to use."), Property);
+				USD_LOG_INFO(TEXT("Skipping material property %d as we have no valid data to use."), Property);
 				continue;
 			}
 
@@ -1684,9 +1681,7 @@ namespace UE::UsdShadeConversion::Private
 			{
 				if (NumSamples != SampleSize->X * SampleSize->Y)
 				{
-					UE_LOG(
-						LogUsd,
-						Warning,
+					USD_LOG_WARNING(
 						TEXT("Skipping material property %d as it has an unexpected number of samples (%d instead of %d)."),
 						Property,
 						NumSamples,
@@ -1698,9 +1693,7 @@ namespace UE::UsdShadeConversion::Private
 
 			if (!bPropertyValueIsConstant && (TextureFilePath == nullptr || !FPaths::FileExists(*TextureFilePath)))
 			{
-				UE_LOG(
-					LogUsd,
-					Warning,
+				USD_LOG_WARNING(
 					TEXT("Skipping material property %d as target texture '%s' was not found."),
 					Property,
 					TextureFilePath ? **TextureFilePath : TEXT("")
@@ -2098,13 +2091,10 @@ namespace UE::UsdShadeConversion::Private
 			}
 			else
 			{
-				FUsdLogManager::LogMessage(
-					EMessageSeverity::Warning,
-					FText::Format(
-						LOCTEXT("FailedToParseTexture", "Failed to parse texture at path '{0}'"),
-						FText::FromString(TextureParameter->TextureFilePath)
-					)
-				);
+				USD_LOG_USERWARNING(FText::Format(
+					LOCTEXT("FailedToParseTexture", "Failed to parse texture at path '{0}'"),
+					FText::FromString(TextureParameter->TextureFilePath)
+				));
 			}
 		}
 	}
@@ -2354,9 +2344,7 @@ bool UsdToUnreal::ConvertMaterial(
 				}
 				else
 				{
-					UE_LOG(
-						LogUsd,
-						Warning,
+					USD_LOG_WARNING(
 						TEXT("Failed to find primvar '%s' when setting material parameter. Available primvars and UV indices: %s.%s"),
 						*TextureParameterValue->Primvar,
 						*UsdUtils::StringifyMap(MaterialData.PrimvarToUVIndex),
@@ -2707,7 +2695,7 @@ FString UsdUtils::GetTextureHash(
 	}
 	else
 	{
-		UE_LOG(LogUsd, Warning, TEXT("Failed to find texture at path '%s' when trying to generate a hash for it"), *ResolvedTexturePath);
+		USD_LOG_WARNING(TEXT("Failed to find texture at path '%s' when trying to generate a hash for it"), *ResolvedTexturePath);
 	}
 
 	// Hash the additional data
@@ -2776,17 +2764,14 @@ void UsdUtils::NotifyIfVirtualTexturesNeeded(UTexture* Texture)
 
 	if (!UseVirtualTexturing(GMaxRHIShaderPlatform))
 	{
-		FUsdLogManager::LogMessage(
-			EMessageSeverity::Warning,
-			FText::Format(
-				LOCTEXT(
-					"DisabledVirtualTexturing",
-					"Texture '{0}' (from prim '{1}') requires Virtual Textures, but the feature is disabled for this project"
-				),
-				FText::FromString(Texture->GetName()),
-				FText::FromString(TexturePath)
-			)
-		);
+		USD_LOG_USERWARNING(FText::Format(
+			LOCTEXT(
+				"DisabledVirtualTexturing",
+				"Texture '{0}' (from prim '{1}') requires Virtual Textures, but the feature is disabled for this project"
+			),
+			FText::FromString(Texture->GetName()),
+			FText::FromString(TexturePath)
+		));
 	}
 }
 
@@ -2886,15 +2871,6 @@ bool UsdUtils::MarkMaterialPrimWithWorldSpaceNormals(const UE::FUsdPrim& Materia
 
 void UsdUtils::SetScalarParameterValue(UMaterialInstance& Material, const TCHAR* ParameterName, float ParameterValue)
 {
-	UE_LOG(
-		LogUsd,
-		Verbose,
-		TEXT("Setting material '%s' scalar parameter '%s' with value '%f'"),
-		*Material.GetPathName(),
-		ParameterName ? ParameterName : TEXT("nullptr"),
-		ParameterValue
-	);
-
 	FMaterialParameterInfo Info;
 	Info.Name = ParameterName;
 
@@ -2912,15 +2888,6 @@ void UsdUtils::SetScalarParameterValue(UMaterialInstance& Material, const TCHAR*
 
 void UsdUtils::SetVectorParameterValue(UMaterialInstance& Material, const TCHAR* ParameterName, FLinearColor ParameterValue)
 {
-	UE_LOG(
-		LogUsd,
-		Verbose,
-		TEXT("Setting material '%s' vector parameter '%s' with value '%s'"),
-		*Material.GetPathName(),
-		ParameterName ? ParameterName : TEXT("nullptr"),
-		*ParameterValue.ToString()
-	);
-
 	FMaterialParameterInfo Info;
 	Info.Name = ParameterName;
 
@@ -2938,15 +2905,6 @@ void UsdUtils::SetVectorParameterValue(UMaterialInstance& Material, const TCHAR*
 
 void UsdUtils::SetTextureParameterValue(UMaterialInstance& Material, const TCHAR* ParameterName, UTexture* ParameterValue)
 {
-	UE_LOG(
-		LogUsd,
-		Verbose,
-		TEXT("Setting material '%s' texture parameter '%s' with value '%s'"),
-		*Material.GetPathName(),
-		ParameterName ? ParameterName : TEXT("nullptr"),
-		*ParameterValue->GetPathName()
-	);
-
 	FMaterialParameterInfo Info;
 	Info.Name = ParameterName;
 
@@ -2964,15 +2922,6 @@ void UsdUtils::SetTextureParameterValue(UMaterialInstance& Material, const TCHAR
 
 void UsdUtils::SetBoolParameterValue(UMaterialInstance& Material, const TCHAR* ParameterName, bool bParameterValue)
 {
-	UE_LOG(
-		LogUsd,
-		Verbose,
-		TEXT("Setting material '%s' bool parameter '%s' with value '%d'"),
-		*Material.GetPathName(),
-		ParameterName ? ParameterName : TEXT("nullptr"),
-		bParameterValue
-	);
-
 	bool bFound = false;
 
 #if WITH_EDITOR
@@ -3118,9 +3067,7 @@ void UsdUtils::AuthorUnrealMaterialBinding(pxr::UsdPrim& MeshOrGeomSubsetPrim, c
 	pxr::UsdShadeMaterial ChildMaterial = pxr::UsdShadeMaterial::Define(Stage, MaterialPath);
 	if (!ChildMaterial)
 	{
-		UE_LOG(
-			LogUsd,
-			Warning,
+		USD_LOG_WARNING(
 			TEXT("Failed to author material prim '%s' when trying to write '%s's material assignment '%s' to USD"),
 			*UsdToUnreal::ConvertPath(MaterialPath),
 			*UsdToUnreal::ConvertPath(MeshOrGeomSubsetPrim.GetPath()),
@@ -3156,9 +3103,7 @@ void UsdUtils::AuthorUnrealCollectionBasedMaterialBinding(
 	{
 		if (!Iter || Iter.IsPseudoRoot())
 		{
-			UE_LOG(
-				LogUsd,
-				Warning,
+			USD_LOG_WARNING(
 				TEXT(
 					"Failed to author collection-based material bindings on prim '%s' for target prim '%s', as the latter is not a descendent of the former"
 				),
@@ -3242,9 +3187,7 @@ void UsdUtils::AuthorUnrealCollectionBasedMaterialBinding(
 		BoundMaterial = pxr::UsdShadeMaterial::Define(Stage, MaterialPath);
 		if (!BoundMaterial)
 		{
-			UE_LOG(
-				LogUsd,
-				Warning,
+			USD_LOG_WARNING(
 				TEXT("Failed to author material prim '%s' when trying to write '%s's collection-based material assignment '%s' to USD"),
 				*UsdToUnreal::ConvertPath(MaterialPath),
 				*UsdToUnreal::ConvertPath(CollectionPrim.GetPath()),
@@ -3324,9 +3267,7 @@ void UsdUtils::AuthorUnrealCollectionBasedMaterialBinding(
 		const bool bCreatedBinding = BindingAPI.Bind(CollectionToUse, BoundMaterial, Empty, BindingStrength);
 		if (!bCreatedBinding)
 		{
-			UE_LOG(
-				LogUsd,
-				Warning,
+			USD_LOG_WARNING(
 				TEXT("Failed to create collection-based material binding '%s' on prim '%s'"),
 				*UsdToUnreal::ConvertToken(CollectionToUse.GetName()),
 				*UsdToUnreal::ConvertPath(CollectionPrim.GetPath())
@@ -3341,9 +3282,7 @@ void UsdUtils::AuthorUnrealCollectionBasedMaterialBinding(
 	const bool bAddedPrim = BindingAPI.AddPrimToBindingCollection(TargetMeshOrGeomSubsetPrim, CollectionName, BindingPurpose);
 	if (!bAddedPrim)
 	{
-		UE_LOG(
-			LogUsd,
-			Warning,
+		USD_LOG_WARNING(
 			TEXT("Failed to add prim '%s' to the collection-based material binding '%s' on prim '%s'"),
 			*UsdToUnreal::ConvertPath(TargetMeshOrGeomSubsetPrim.GetPath()),
 			*UsdToUnreal::ConvertToken(CollectionName),
diff --git a/Engine/Plugins/Runtime/USDCore/Source/USDUtilities/Private/USDSkeletalDataConversion.cpp b/Engine/Plugins/Runtime/USDCore/Source/USDUtilities/Private/USDSkeletalDataConversion.cpp
--- a/Engine/Plugins/Runtime/USDCore/Source/USDUtilities/Private/USDSkeletalDataConversion.cpp
+++ b/Engine/Plugins/Runtime/USDCore/Source/USDUtilities/Private/USDSkeletalDataConversion.cpp
@@ -10,7 +10,6 @@
 #include "USDErrorUtils.h"
 #include "USDGeomMeshConversion.h"
 #include "USDLayerUtils.h"
-#include "USDLog.h"
 #include "USDMemory.h"
 #include "USDObjectUtils.h"
 #include "USDPrimConversion.h"
@@ -301,33 +300,37 @@ namespace SkelDataConversionImpl
 
 		if (NumNormals > 0 && NumOffsets != NumNormals)
 		{
-			UE_LOG(
-				LogUsd,
-				Warning,
-				TEXT("BlendShape '%s' has mismatching numbers of offsets (%d) and normalOffsets (%d) and will be ignored"),
-				*Name,
+			USD_LOG_USERWARNING(FText::Format(
+				LOCTEXT(
+					"OffsetNormalMismach",
+					"BlendShape '{0}' has mismatching numbers of offsets ({1}) and normalOffsets ({2}) and will be ignored"
+				),
+				FText::FromString(Name),
 				NumOffsets,
 				NumNormals
-			);
+			));
 			return false;
 		}
 
 		if (NumIndices > 0 && NumOffsets != NumIndices)
 		{
-			UE_LOG(
-				LogUsd,
-				Warning,
-				TEXT("BlendShape '%s' has mismatching numbers of offsets (%d) and point indices (%d) and will be ignored"),
-				*Name,
+			USD_LOG_USERWARNING(FText::Format(
+				LOCTEXT(
+					"OffsetPointMismatch",
+					"BlendShape '{0}' has mismatching numbers of offsets ({1}) and point indices ({2}) and will be ignored"
+				),
+				FText::FromString(Name),
 				NumOffsets,
 				NumIndices
-			);
+			));
 			return false;
 		}
 
 		if (NumOffsets + NumNormals == 0)
 		{
-			UE_LOG(LogUsd, Warning, TEXT("BlendShape '%s' zero offsets and normalOffsets and will be ignored"), *Name);
+			USD_LOG_USERWARNING(
+				FText::Format(LOCTEXT("NoOffsets", "BlendShape '{0}' zero offsets and normalOffsets and will be ignored"), FText::FromString(Name))
+			);
 			return false;
 		}
 
@@ -478,17 +481,14 @@ namespace SkelDataConversionImpl
 		{
 			if (!(Curve->FloatCurve == SourceData))
 			{
-				FUsdLogManager::LogMessage(
-					EMessageSeverity::Warning,
-					FText::Format(
-						LOCTEXT(
-							"OverwritingMorphTargetCurves",
-							"Overwriting animation curve for morph target '{0}' with different data! If the Skeletal Mesh has multiple LODs, make "
-							"sure each LOD mesh that wants to animate a certain blend shape does so with the same blend shape curve."
-						),
-						FText::FromName(CurveName)
-					)
-				);
+				USD_LOG_USERWARNING(FText::Format(
+					LOCTEXT(
+						"OverwritingMorphTargetCurves",
+						"Overwriting animation curve for morph target '{0}' with different data! If the Skeletal Mesh has multiple LODs, make "
+						"sure each LOD mesh that wants to animate a certain blend shape does so with the same blend shape curve."
+					),
+					FText::FromName(CurveName)
+				));
 			}
 
 			Controller.SetCurveFlags(CurveId, Curve->GetCurveTypeFlags() | AACF_DefaultCurve, bShouldTransact);
@@ -500,13 +500,7 @@ namespace SkelDataConversionImpl
 		}
 		else
 		{
-			UE_LOG(
-				LogUsd,
-				Error,
-				TEXT("Failed to create float curve with name '%s' for UAnimSequence '%s'"),
-				*CurveName.ToString(),
-				*Sequence->GetName()
-			);
+			USD_LOG_ERROR(TEXT("Failed to create float curve with name '%s' for UAnimSequence '%s'"), *CurveName.ToString(), *Sequence->GetName());
 		}
 	}
 
@@ -1127,18 +1121,15 @@ bool UsdToUnreal::ConvertSkeleton(
 	{
 		FString SkeletonPrimPath = UsdToUnreal::ConvertPath(InSkeletonQuery.GetPrim().GetPath());
 
-		FUsdLogManager::LogMessage(
-			EMessageSeverity::Warning,
-			FText::Format(
-				LOCTEXT(
-					"NoBonesInSkeleton",
-					"Skeleton prim '{0}' has no joints! "
-					"A new skeleton with a single root joint will be generated as USkeletalMeshes require valid skeletons. "
-					"Note that this new skeleton may be written back to the USD stage when exporting the corresponding asset."
-				),
-				FText::FromString(SkeletonPrimPath)
-			)
-		);
+		USD_LOG_USERWARNING(FText::Format(
+			LOCTEXT(
+				"NoBonesInSkeleton",
+				"Skeleton prim '{0}' has no joints! "
+				"A new skeleton with a single root joint will be generated as USkeletalMeshes require valid skeletons. "
+				"Note that this new skeleton may be written back to the USD stage when exporting the corresponding asset."
+			),
+			FText::FromString(SkeletonPrimPath)
+		));
 
 		if (bEnsureAtLeastOneBone)
 		{
@@ -1418,7 +1409,8 @@ bool UsdToUnreal::ConvertSkinnedMesh(
 		for (uint32 PointIndex = 0; PointIndex < NumPoints; ++PointIndex)
 		{
 			const GfVec3f& Point = UsdPoints[PointIndex];
-			SkelMeshImportData.Points[PointIndex + NumExistingPoints] = (FVector3f)CommonOptions.AdditionalTransform.TransformPosition(UsdToUnreal::ConvertVector(StageInfo, Point));
+			SkelMeshImportData.Points[PointIndex + NumExistingPoints] = (FVector3f
+			)CommonOptions.AdditionalTransform.TransformPosition(UsdToUnreal::ConvertVector(StageInfo, Point));
 		}
 	}
 
@@ -1594,15 +1586,16 @@ bool UsdToUnreal::ConvertSkinnedMesh(
 			}
 			else
 			{
-				UE_LOG(
-					LogUsd,
-					Warning,
-					TEXT("Prim '%s' has invalid number of displayColor values for primvar interpolation type '%s'! (expected %d, found %d)"),
-					*UsdToUnreal::ConvertPath(SkinningPrim.GetPath()),
-					*UsdToUnreal::ConvertToken(USDInterpType),
+				USD_LOG_USERWARNING(FText::Format(
+					LOCTEXT(
+						"InvalidDisplayColorInterpolation",
+						"Prim '{0}' has invalid number of displayColor values for primvar interpolation type '{1}'! (expected {2}, found {3})"
+					),
+					FText::FromString(UsdToUnreal::ConvertPath(SkinningPrim.GetPath())),
+					FText::FromString(UsdToUnreal::ConvertToken(USDInterpType)),
 					NumExpectedColors,
 					NumColors
-				);
+				));
 			}
 		}
 	}
@@ -1653,15 +1646,16 @@ bool UsdToUnreal::ConvertSkinnedMesh(
 			}
 			else
 			{
-				UE_LOG(
-					LogUsd,
-					Warning,
-					TEXT("Prim '%s' has invalid number of displayOpacity values for primvar interpolation type '%s'! (expected %d, found %d)"),
-					*UsdToUnreal::ConvertPath(SkinningPrim.GetPath()),
-					*UsdToUnreal::ConvertToken(USDInterpType),
+				USD_LOG_USERWARNING(FText::Format(
+					LOCTEXT(
+						"InvalidDisplayOpacityInterpolation",
+						"Prim '{0}' has invalid number of displayOpacity values for primvar interpolation type '{1}'! (expected {2}, found {3})"
+					),
+					FText::FromString(UsdToUnreal::ConvertPath(SkinningPrim.GetPath())),
+					FText::FromString(UsdToUnreal::ConvertToken(USDInterpType)),
 					NumExpectedOpacities,
 					NumOpacities
-				);
+				));
 			}
 		}
 	}
@@ -1964,9 +1958,7 @@ bool UsdToUnreal::ConvertSkinnedMesh(
 
 	if (NumSkippedPolygons > 0)
 	{
-		UE_LOG(
-			LogUsd,
-			Warning,
+		USD_LOG_WARNING(
 			TEXT("Ignoring %d polygons with less than 3 vertices from mesh '%s'"),
 			NumSkippedPolygons,
 			*UsdToUnreal::ConvertPath(UsdMesh.GetPrim().GetPrimPath())
@@ -2140,10 +2132,7 @@ bool UsdToUnreal::ConvertSkelAnim(
 	double StageTimeCodesPerSecond = Stage.Get()->GetTimeCodesPerSecond();
 	if (FMath::IsNearlyZero(StageTimeCodesPerSecond))
 	{
-		FUsdLogManager::LogMessage(
-			EMessageSeverity::Warning,
-			LOCTEXT("TimeCodesPerSecondIsZero", "Cannot bake skeletal animations as the stage has timeCodesPerSecond set to zero!")
-		);
+		USD_LOG_USERWARNING(LOCTEXT("TimeCodesPerSecondIsZero", "Cannot bake skeletal animations as the stage has timeCodesPerSecond set to zero!"));
 		return false;
 	}
 
@@ -2444,13 +2433,10 @@ bool UsdToUnreal::ConvertSkelAnim(
 
 						if (!PrimaryBlendShapeCurve)
 						{
-							FUsdLogManager::LogMessage(
-								EMessageSeverity::Warning,
-								FText::Format(
-									LOCTEXT("NoChannelForPrimary", "Could not find a float channel to apply to primary blend shape '{0}'"),
-									FText::FromString(PrimaryBlendShapePath)
-								)
-							);
+							USD_LOG_USERWARNING(FText::Format(
+								LOCTEXT("NoChannelForPrimary", "Could not find a float channel to apply to primary blend shape '{0}'"),
+								FText::FromString(PrimaryBlendShapePath)
+							));
 							continue;
 						}
 
@@ -2473,13 +2459,10 @@ bool UsdToUnreal::ConvertSkelAnim(
 							);
 							if (RemappedBlendShapeCurves.Num() != FoundPrimaryBlendShape->Inbetweens.Num() + 1)
 							{
-								FUsdLogManager::LogMessage(
-									EMessageSeverity::Warning,
-									FText::Format(
-										LOCTEXT("FailedToRemapInbetweens", "Failed to remap inbetween float curves for blend shape '{0}'"),
-										FText::FromString(PrimaryBlendShapePath)
-									)
-								);
+								USD_LOG_USERWARNING(FText::Format(
+									LOCTEXT("FailedToRemapInbetweens", "Failed to remap inbetween float curves for blend shape '{0}'"),
+									FText::FromString(PrimaryBlendShapePath)
+								));
 								continue;
 							}
 
@@ -2657,8 +2640,6 @@ bool UsdToUnreal::ConvertBlendShape(
 
 		if (Weight > 1.0f || Weight < 0.0f || FMath::IsNearlyZero(Weight) || FMath::IsNearlyEqual(Weight, 1.0f))
 		{
-			// UE_LOG(LogUsd, Warning, TEXT("Inbetween shape '%s' for blend shape '%s' has invalid weight '%f' and will be ignored. Valid weights are
-			// > 0.0f and < 1.0f!"), 	*OrigInbetweenName, 	*PrimaryPath, 	Weight);
 			continue;
 		}
 
@@ -2883,11 +2864,11 @@ bool UsdToUnreal::ConvertSkeletalImportData(
 
 			if (bBuildSuccess)
 			{
-				UE_LOG(LogUsd, Warning, TEXT("Warning when trying to build skeletal mesh from USD: '%s': '%s'"), *Name.ToString(), *Text.ToString());
+				USD_LOG_WARNING(TEXT("Warning when trying to build skeletal mesh from USD: '%s': '%s'"), *Name.ToString(), *Text.ToString());
 			}
 			else
 			{
-				UE_LOG(LogUsd, Error, TEXT("Error when trying to build skeletal mesh from USD: '%s': '%s'"), *Name.ToString(), *Text.ToString());
+				USD_LOG_ERROR(TEXT("Error when trying to build skeletal mesh from USD: '%s': '%s'"), *Name.ToString(), *Text.ToString());
 			}
 		}
 
@@ -3361,9 +3342,7 @@ bool UsdUtils::ApplyBlendShape(
 			}
 			else
 			{
-				UE_LOG(
-					LogUsd,
-					Warning,
+				USD_LOG_WARNING(
 					TEXT("Failed to find inbetween '%s' when applying blend shape prim '%s' to a mesh description"),
 					*InInbetweenName,
 					*UsdToUnreal::ConvertPath(InBlendShapePrim.GetPrimPath())
@@ -3407,9 +3386,7 @@ bool UsdUtils::ApplyBlendShape(
 		else
 		{
 			FString InbetweenText = InInbetweenName.IsEmpty() ? FString::Printf(TEXT(" (inbetween '%s')"), *InInbetweenName) : TEXT("");
-			UE_LOG(
-				LogUsd,
-				Warning,
+			USD_LOG_WARNING(
 				TEXT("Failed to apply position offsets from BlendShape '%s'%s: Expected MeshDescription to have %d vertex positions, but it has %u!"),
 				*UsdToUnreal::ConvertPath(InBlendShapePrim.GetPrimPath()),
 				*InbetweenText,
@@ -3440,9 +3417,7 @@ bool UsdUtils::ApplyBlendShape(
 		else if (NormalOffsets.size() != 0)
 		{
 			FString InbetweenText = InInbetweenName.IsEmpty() ? FString::Printf(TEXT(" (inbetween '%s')"), *InInbetweenName) : TEXT("");
-			UE_LOG(
-				LogUsd,
-				Warning,
+			USD_LOG_WARNING(
 				TEXT("Failed to apply normal offsets from BlendShape '%s'%s: Expected MeshDescription to have %d vertices, but it has %u!"),
 				*UsdToUnreal::ConvertPath(InBlendShapePrim.GetPrimPath()),
 				*InbetweenText,
@@ -3472,9 +3447,7 @@ bool UsdUtils::ApplyBlendShape(
 		else
 		{
 			FString InbetweenText = InInbetweenName.IsEmpty() ? FString::Printf(TEXT(" (inbetween '%s')"), *InInbetweenName) : TEXT("");
-			UE_LOG(
-				LogUsd,
-				Warning,
+			USD_LOG_WARNING(
 				TEXT(
 					"Failed to apply indexed position offsets from BlendShape '%s'%s: The blend shape has %u offsets, but %u indices! (those should match)"
 				),
@@ -3511,9 +3484,7 @@ bool UsdUtils::ApplyBlendShape(
 		else if (NormalOffsets.size() != 0)
 		{
 			FString InbetweenText = InInbetweenName.IsEmpty() ? FString::Printf(TEXT(" (inbetween '%s')"), *InInbetweenName) : TEXT("");
-			UE_LOG(
-				LogUsd,
-				Warning,
+			USD_LOG_WARNING(
 				TEXT(
 					"Failed to apply indexed normal offsets from BlendShape '%s'%s: The blend shape has %u offsets, but %u indices! (those should match)"
 				),
@@ -3668,9 +3639,7 @@ bool UnrealToUsd::ConvertSkeletalMesh(
 	pxr::UsdVariantSets VariantSets = SkelRootPrim.GetVariantSets();
 	if (NumLODs > 1 && VariantSets.HasVariantSet(UnrealIdentifiers::LOD))
 	{
-		UE_LOG(
-			LogUsd,
-			Error,
+		USD_LOG_ERROR(
 			TEXT("Failed to export higher LODs for skeletal mesh '%s', as the target prim already has a variant set named '%s'!"),
 			*SkeletalMesh->GetName(),
 			*UsdToUnreal::ConvertToken(UnrealIdentifiers::LOD)
@@ -3934,7 +3903,7 @@ bool UnrealToUsd::ConvertAnimSequence(UAnimSequence* AnimSequence, pxr::UsdPrim&
 
 			FBlendedCurve BlendedCurve;
 			constexpr bool bForceUseRawData = true;
-			AnimSequence->EvaluateCurveData(BlendedCurve,  FAnimExtractContext(static_cast<double>(AnimTime)), bForceUseRawData);
+			AnimSequence->EvaluateCurveData(BlendedCurve, FAnimExtractContext(static_cast<double>(AnimTime)), bForceUseRawData);
 
 			BlendShapeNames.clear();
 			BlendShapeNames.reserve(BlendedCurve.Num());
@@ -4311,10 +4280,8 @@ bool UnrealToUsd::ConvertControlRigSection(
 							{
 								bRenamedAChannel = true;
 								BlendShapeChannels[BlendShapeIndex] = UnrealToUsd::ConvertToken(*FoundBlendShape->Name).Get();
-								UE_LOG(
-									LogUsd,
-									Log,
-									TEXT("Updating Mesh '%s' to bind BlendShape target '%s' to SkelAnimation curve '%s'"),
+								USD_LOG_INFO(
+									TEXT("Updating Mesh '{0}' to bind BlendShape target '{1}' to SkelAnimation curve '{2}'"),
 									*MeshPrim.GetPrimPath().GetString(),
 									*PrimaryBlendShapePath,
 									*FoundBlendShape->Name
@@ -4322,15 +4289,16 @@ bool UnrealToUsd::ConvertControlRigSection(
 
 								if (FoundBlendShape->Inbetweens.Num() > 0)
 								{
-									UE_LOG(
-										LogUsd,
-										Warning,
-										TEXT("Baking Control Rig parameter sections for BlendShapes with inbetweens (like '%s') is not currently "
-											 "supported, so animation for mesh '%s' may look incorrect! Please flatten the inbetweens into separate "
-											 "BlendShapes beforehand (importing and exporting will do that)."),
-										*FoundBlendShape->Name,
-										*MeshPrim.GetPrimPath().GetString()
-									);
+									USD_LOG_USERWARNING(FText::Format(
+										LOCTEXT(
+											"UnsupportedInbetweens",
+											"Baking Control Rig parameter sections for BlendShapes with inbetweens (like '{0}') is not currently "
+											"supported, so animation for mesh '{1}' may look incorrect! Please flatten the inbetweens into separate "
+											"BlendShapes beforehand (importing and exporting will do that)."
+										),
+										FText::FromString(FoundBlendShape->Name),
+										FText::FromString(MeshPrim.GetPrimPath().GetString())
+									));
 								}
 							}
 						}
@@ -4377,15 +4345,13 @@ bool UnrealToUsd::ConvertControlRigSection(
 		URigHierarchy* Hierarchy = ControlRig->GetHierarchy();
 		if (!Hierarchy)
 		{
-			UE_LOG(LogUsd, Error, TEXT("Baking Control Rig tracks for rig '%s' failed"), *ControlRig->GetPathName());
+			USD_LOG_ERROR(TEXT("Baking Control Rig tracks for rig '%s' failed"), *ControlRig->GetPathName());
 			return false;
 		}
 
 		if (Hierarchy->GetTopologyVersion() != TopologyVersion)
 		{
-			UE_LOG(
-				LogUsd,
-				Log,
+			USD_LOG_INFO(
 				TEXT("Regenerating ControlRig to reference skeleton mapping for rig '%s' as its topology changed"),
 				*ControlRig->GetPathName()
 			);
@@ -4450,9 +4416,7 @@ bool UnrealToUsd::ConvertControlRigSection(
 	double ElapsedSeconds = FPlatformTime::ToSeconds64(FPlatformTime::Cycles64() - StartTime);
 	int ElapsedMin = int(ElapsedSeconds / 60.0);
 	ElapsedSeconds -= 60.0 * (double)ElapsedMin;
-	UE_LOG(
-		LogUsd,
-		Log,
+	USD_LOG_INFO(
 		TEXT("Baked new animation for prim '%s' in [%d min %.3f s]"),
 		*UsdToUnreal::ConvertPath(OutSkelAnimPrim.GetPrimPath()),
 		ElapsedMin,
diff --git a/Engine/Plugins/Runtime/USDCore/Source/USDUtilities/Private/USDUtilitiesModule.cpp b/Engine/Plugins/Runtime/USDCore/Source/USDUtilities/Private/USDUtilitiesModule.cpp
--- a/Engine/Plugins/Runtime/USDCore/Source/USDUtilities/Private/USDUtilitiesModule.cpp
+++ b/Engine/Plugins/Runtime/USDCore/Source/USDUtilities/Private/USDUtilitiesModule.cpp
@@ -2,6 +2,7 @@
 
 #include "USDUtilitiesModule.h"
 
+#include "USDErrorUtils.h"
 #include "USDMemory.h"
 
 #include "Modules/ModuleManager.h"
@@ -12,11 +13,43 @@
 
 class FUsdUtilitiesModule : public IUsdUtilitiesModule
 {
+	FDelegateHandle OnProjectSettingsChangedHandle;
+
+public:
 	virtual void StartupModule() override
 	{
+		UUsdProjectSettings* ProjectSettings = GetMutableDefault<UUsdProjectSettings>();
+		if (ProjectSettings->bLogUsdSdkErrors)
+		{
+			FUsdLogManager::RegisterDiagnosticDelegate();
+		}
+
 #if WITH_EDITOR
 		LLM_SCOPE_BYTAG(Usd);
 
+		OnProjectSettingsChangedHandle = ProjectSettings->OnSettingChanged().AddLambda(
+			[](UObject* SettingsObject, struct FPropertyChangedEvent& PropertyChangedEvent)
+			{
+				UUsdProjectSettings* UsdSettings = Cast<UUsdProjectSettings>(SettingsObject);
+				if (!UsdSettings)
+				{
+					return;
+				}
+
+				if (PropertyChangedEvent.GetPropertyName() == GET_MEMBER_NAME_CHECKED(UUsdProjectSettings, bLogUsdSdkErrors))
+				{
+					if (UsdSettings->bLogUsdSdkErrors)
+					{
+						FUsdLogManager::RegisterDiagnosticDelegate();
+					}
+					else
+					{
+						FUsdLogManager::UnregisterDiagnosticDelegate();
+					}
+				}
+			}
+		);
+
 		FMessageLogInitializationOptions InitOptions;
 		FMessageLogModule& MessageLogModule = FModuleManager::LoadModuleChecked<FMessageLogModule>(TEXT("MessageLog"));
 		MessageLogModule.RegisterLogListing(TEXT("USD"), NSLOCTEXT("USDUtilitiesModule", "USDLogListing", "USD"), InitOptions);
@@ -25,7 +58,19 @@ class FUsdUtilitiesModule : public IUsdUtilitiesModule
 
 	virtual void ShutdownModule() override
 	{
+		FUsdLogManager::UnregisterDiagnosticDelegate();
+
 #if WITH_EDITOR
+		// We can't query default objects during engine exit
+		if (UObjectInitialized())
+		{
+			UUsdProjectSettings* ProjectSettings = GetMutableDefault<UUsdProjectSettings>();
+			if (ProjectSettings)
+			{
+				ProjectSettings->OnSettingChanged().Remove(OnProjectSettingsChangedHandle);
+			}
+		}
+
 		FMessageLogModule& MessageLogModule = FModuleManager::LoadModuleChecked<FMessageLogModule>(TEXT("MessageLog"));
 		MessageLogModule.UnregisterLogListing(TEXT("USD"));
 #endif	  // WITH_EDITOR
diff --git a/Engine/Plugins/Runtime/USDCore/Source/USDUtilities/Private/USDValueConversion.cpp b/Engine/Plugins/Runtime/USDCore/Source/USDUtilities/Private/USDValueConversion.cpp
--- a/Engine/Plugins/Runtime/USDCore/Source/USDUtilities/Private/USDValueConversion.cpp
+++ b/Engine/Plugins/Runtime/USDCore/Source/USDUtilities/Private/USDValueConversion.cpp
@@ -3,7 +3,7 @@
 #include "USDValueConversion.h"
 
 #include "USDConversionUtils.h"
-#include "USDLog.h"
+#include "USDErrorUtils.h"
 #include "USDMemory.h"
 #include "USDTypesConversion.h"
 #include "UsdWrappers/UsdStage.h"
@@ -2673,12 +2673,12 @@ namespace UsdUtils
 			}
 			else
 			{
-				UE_LOG(LogUsd, Error, TEXT("Found unstringifier for typeName '%s', but failed to unstringify string '%s'"), *TypeName, *String);
+				USD_LOG_ERROR(TEXT("Found unstringifier for typeName '%s', but failed to unstringify string '%s'"), *TypeName, *String);
 			}
 		}
 		else
 		{
-			UE_LOG(LogUsd, Error, TEXT("Failed to find unstringifier for typeName '%s'"), *TypeName);
+			USD_LOG_ERROR(TEXT("Failed to find unstringifier for typeName '%s'"), *TypeName);
 		}
 
 		return false;
diff --git a/Engine/Plugins/Runtime/USDCore/Source/USDUtilities/Public/USDErrorUtils.h b/Engine/Plugins/Runtime/USDCore/Source/USDUtilities/Public/USDErrorUtils.h
--- a/Engine/Plugins/Runtime/USDCore/Source/USDUtilities/Public/USDErrorUtils.h
+++ b/Engine/Plugins/Runtime/USDCore/Source/USDUtilities/Public/USDErrorUtils.h
@@ -4,18 +4,21 @@
 
 #include "CoreMinimal.h"
 #include "Logging/TokenizedMessage.h"
+#include "Misc/SourceLocation.h"
 
 namespace UsdUtils
 {
 	/**
 	 * Pushes an USD error monitoring object into the stack and catches any emitted errors
 	 */
+	UE_DEPRECATED(5.6, "Use FScopedUsdMessageLog instead.")
 	USDUTILITIES_API void StartMonitoringErrors();
 
 	/**
-	 * Returns all errors that were captured since StartMonitoringErrors(), clears and pops an
-	 * error monitoring object from the stack
+	* Returns all errors that were captured since StartMonitoringErrors(), clears and pops an
+	* error monitoring object from the stack
 	 */
+	UE_DEPRECATED(5.6, "Use FScopedUsdMessageLog and FUsdLogManager::HasAccumulatedErrors() instead.")
 	USDUTILITIES_API TArray<FString> GetErrorsAndStopMonitoring();
 
 	/**
@@ -24,61 +27,84 @@ namespace UsdUtils
 	 * If ToastMessage is empty, a default message will be displayed.
 	 * Returns true if there were any errors.
 	 */
+	UE_DEPRECATED(5.6, "Use FScopedUsdMessageLog instead.")
 	USDUTILITIES_API bool ShowErrorsAndStopMonitoring(const FText& ToastMessage = FText());
 }
 
-class FUsdLogManager;
-
-namespace UE
+class USDUTILITIES_API FUsdLogManager
 {
-	namespace Internal
+public:
+	// Copied from AssertionMacros.h as it's private there
+	UE_NODEBUG static constexpr uint32 FileHashForEnsure(const char* Filename)
 	{
-		class FUsdMessageLog
+		uint32 Result = 5381;
+		for (;;)
 		{
-			friend class ::FUsdLogManager;
-
-		public:
-			FUsdMessageLog() = default;
+			char Ch = *Filename++;
+			if (!Ch)
+			{
+				return Result;
+			}
+			Result = ((Result << 5) + Result) + Ch;
+		}
+	}
 
-			FUsdMessageLog(const FUsdMessageLog&) = default;
-			FUsdMessageLog& operator=(const FUsdMessageLog&) = default;
-			FUsdMessageLog(FUsdMessageLog&&) = default;
-			FUsdMessageLog& operator=(FUsdMessageLog&&) = default;
+	// Log a message that shows only on the output log.
+	// Prefer using the USD_LOG_INFO, USD_LOG_WARNING, USD_LOG_ERROR, USD_LOG_USERINFO, USD_LOG_USERWARNING, USD_LOG_USERERROR macros.
+	static void Log(EMessageSeverity::Type Severity, const FString& Message, uint32 MessageID);
 
-			~FUsdMessageLog();
+	// Log a message that shows on the output log and the message log, if we're inside of a FScopedUsdMessageLog scope.
+	// Prefer using the USD_LOG_INFO, USD_LOG_WARNING, USD_LOG_ERROR, USD_LOG_USERINFO, USD_LOG_USERWARNING, USD_LOG_USERERROR macros.
+	static void Log(EMessageSeverity::Type Severity, const FText& Message, uint32 MessageID);
 
-		private:
-			/** Display pending messages, show the log window */
-			void Dump();
+	// Returns whether we have currently accumulated any message with severity warning or higher.
+	// If we're not in the scope of any FScopedUsdMessageLog, returns false.
+	static bool HasAccumulatedErrors();
 
-			/** Append a message */
-			void Push(const TSharedRef<FTokenizedMessage>& Message);
+	static void RegisterDiagnosticDelegate();
 
-		private:
-			TArray<TSharedRef<FTokenizedMessage>> TokenizedMessages;
-		};
-	}
-}	 // namespace UE
+	static void UnregisterDiagnosticDelegate();
 
-class USDUTILITIES_API FUsdLogManager
-{
-public:
-	/** Sends the message to Message Log if it exists, if not, sends it to the Output Log */
+	UE_DEPRECATED(5.6, "Use Log, or ideally the USD_LOG_INFO, USD_LOG_WARNING, etc. macros")
 	static void LogMessage(EMessageSeverity::Type Severity, const FText& Message);
+
+	UE_DEPRECATED(5.6, "Use Log, or ideally the USD_LOG_INFO, USD_LOG_WARNING, etc. macros")
 	static void LogMessage(const TSharedRef<FTokenizedMessage>& Message);
 
+	UE_DEPRECATED(5.6, "Used FScopedUsdMessageLog instead.")
 	static void EnableMessageLog();
-	static void DisableMessageLog();
-
-private:
-	static TOptional<UE::Internal::FUsdMessageLog> MessageLog;
-	static int32 MessageLogRefCount;
 
-	static FCriticalSection MessageLogLock;
+	UE_DEPRECATED(5.6, "Used FScopedUsdMessageLog instead.")
+	static void DisableMessageLog();
 };
 
 /**
- * Starts sending USD messages to the message log and will display the message log on destruction if any messages were logged.
+ * Display an Info/Warn/Error message on the Output Log (USD_LOG_INFO, USD_LOG_WARNING, USD_LOG_ERROR),
+ * or display a user-facing Info/Warn/Error message, that is additionally added to the Message Log
+ * (USD_LOG_USERINFO, USD_LOG_USERWARNING, USD_LOG_USERERROR).
+ *
+ * The user-facing message must be within the scope of a FScopedUsdMessageLog to be added to the Message Log, otherwise it is treated as the
+ * non-user analogue. See FScopedUsdMessageLog right below this for more info.
+ *
+ * Usage examples:
+ *     USD_LOG_WARNING(TEXT("Some error directly to the Output Log: %s"), *MyString);
+ *     USD_LOG_USERERROR(LOCTEXT("MyMessage", "This is an error message that will show on the Message Log"));
+ *     USD_LOG_USERWARNING(FText::Format(LOCTEXT("InvalidBinding", "Binding {0} is invalid!"), FText::FromString(UsdToUnreal::ConvertPath(InSkeletonRoot.GetPath()))));
+ */
+#define USD_LOG_INFO(Format, ...) 		FUsdLogManager::Log(EMessageSeverity::Info, 	FString::Printf(Format, ##__VA_ARGS__), HashCombine(FUsdLogManager::FileHashForEnsure(__FILE__), __LINE__));
+#define USD_LOG_WARNING(Format, ...) 	FUsdLogManager::Log(EMessageSeverity::Warning, 	FString::Printf(Format, ##__VA_ARGS__), HashCombine(FUsdLogManager::FileHashForEnsure(__FILE__), __LINE__));
+#define USD_LOG_ERROR(Format, ...) 		FUsdLogManager::Log(EMessageSeverity::Error, 	FString::Printf(Format, ##__VA_ARGS__), HashCombine(FUsdLogManager::FileHashForEnsure(__FILE__), __LINE__));
+#define USD_LOG_USERINFO(text) 			FUsdLogManager::Log(EMessageSeverity::Info, 	text, 									HashCombine(FUsdLogManager::FileHashForEnsure(__FILE__), __LINE__));
+#define USD_LOG_USERWARNING(text) 		FUsdLogManager::Log(EMessageSeverity::Warning, 	text, 									HashCombine(FUsdLogManager::FileHashForEnsure(__FILE__), __LINE__));
+#define USD_LOG_USERERROR(text) 		FUsdLogManager::Log(EMessageSeverity::Error, 	text, 									HashCombine(FUsdLogManager::FileHashForEnsure(__FILE__), __LINE__));
+
+/**
+ * Begins a scope where all logged messages (user and non-user facing, USD SDK errors and even USD error mark messages) are collected and deduplicated.
+ *
+ * The deduplication only happens if the USD.CollapseLogMessages cvar is enabled (default).
+ *
+ * If the USD.UseMessageLog cvar is enabled (default), it will add the user-facing messages (and the USD SDK and error mark messages) to the Message Log,
+ * (whether deduplication is enabled or disabled), also displaying a toast letting the user know that those messages can be viewed there.
  */
 class USDUTILITIES_API FScopedUsdMessageLog
 {
diff --git a/Engine/Plugins/Runtime/USDCore/Source/USDUtilities/USDUtilities.Build.cs b/Engine/Plugins/Runtime/USDCore/Source/USDUtilities/USDUtilities.Build.cs
--- a/Engine/Plugins/Runtime/USDCore/Source/USDUtilities/USDUtilities.Build.cs
+++ b/Engine/Plugins/Runtime/USDCore/Source/USDUtilities/USDUtilities.Build.cs
@@ -53,6 +53,7 @@ namespace UnrealBuildTool.Rules
 					new string[]
 					{
 						"DesktopPlatform", // For OpenFileDialog/SaveFileDialog
+						"DeveloperSettings", // So that we can un/register the diagnostic delegate when the project settings change
 						"MaterialBaking", // For the BakeMaterials function
 						"MaterialEditor",
 						"MeshUtilities",
diff --git a/Engine/Plugins/Runtime/USDCore/Source/UnrealUSDWrapper/Private/USDListener.cpp b/Engine/Plugins/Runtime/USDCore/Source/UnrealUSDWrapper/Private/USDListener.cpp
--- a/Engine/Plugins/Runtime/USDCore/Source/UnrealUSDWrapper/Private/USDListener.cpp
+++ b/Engine/Plugins/Runtime/USDCore/Source/UnrealUSDWrapper/Private/USDListener.cpp
@@ -2,7 +2,6 @@
 
 #include "USDListener.h"
 
-#include "USDLog.h"
 #include "USDMemory.h"
 
 #include "UsdWrappers/ForwardDeclarations.h"
diff --git a/Engine/Plugins/Runtime/USDCore/Source/UnrealUSDWrapper/Private/UnrealUSDWrapper.cpp b/Engine/Plugins/Runtime/USDCore/Source/UnrealUSDWrapper/Private/UnrealUSDWrapper.cpp
--- a/Engine/Plugins/Runtime/USDCore/Source/UnrealUSDWrapper/Private/UnrealUSDWrapper.cpp
+++ b/Engine/Plugins/Runtime/USDCore/Source/UnrealUSDWrapper/Private/UnrealUSDWrapper.cpp
@@ -518,58 +518,6 @@ namespace UsdWrapperUtils
 	}
 }
 
-#if USE_USD_SDK
-class FUsdDiagnosticDelegate : public pxr::TfDiagnosticMgr::Delegate
-{
-public:
-	virtual ~FUsdDiagnosticDelegate() override {};
-
-	virtual void IssueError(const pxr::TfError& Error) override
-	{
-		FScopedUsdAllocs Allocs;
-
-		std::string Msg = Error.GetErrorCodeAsString();
-		Msg += ": ";
-		Msg += Error.GetCommentary();
-
-		UE_LOG(LogUsd, Error, TEXT("%s"), UTF8_TO_TCHAR(Msg.c_str()));
-	}
-
-	virtual void IssueFatalError(const pxr::TfCallContext& Context, const std::string& Msg) override
-	{
-		UE_LOG(LogUsd, Error, TEXT("%s"), UTF8_TO_TCHAR(Msg.c_str()));
-	}
-
-	virtual void IssueStatus(const pxr::TfStatus& Status) override
-	{
-		FScopedUsdAllocs Allocs;
-
-		std::string Msg = Status.GetDiagnosticCodeAsString();
-		Msg += ": ";
-		Msg += Status.GetCommentary();
-
-		UE_LOG(LogUsd, Log, TEXT("%s"), UTF8_TO_TCHAR(Msg.c_str()));
-	}
-
-	virtual void IssueWarning(const pxr::TfWarning& Warning) override
-	{
-		FScopedUsdAllocs Allocs;
-
-		std::string Msg = Warning.GetDiagnosticCodeAsString();
-		Msg += ": ";
-		Msg += Warning.GetCommentary();
-
-		UE_LOG(LogUsd, Warning, TEXT("%s"), UTF8_TO_TCHAR(Msg.c_str()));
-	}
-};
-#else
-class FUsdDiagnosticDelegate
-{
-};
-#endif	  // USE_USD_SDK
-
-TUniquePtr<FUsdDiagnosticDelegate> UnrealUSDWrapper::Delegate = nullptr;
-
 #if USE_USD_SDK
 
 double UnrealUSDWrapper::GetDefaultTimeCode()
@@ -1095,79 +1043,24 @@ void UnrealUSDWrapper::SetDefaultResolverDefaultSearchPath(const TArray<FDirecto
 #endif	  // #if USE_USD_SDK
 }
 
-void UnrealUSDWrapper::SetupDiagnosticDelegate() TSAN_SAFE
+// Deprecated
+void UnrealUSDWrapper::SetupDiagnosticDelegate()
 {
-#if USE_USD_SDK
-	if (Delegate.IsValid())
-	{
-		UnrealUSDWrapper::ClearDiagnosticDelegate();
-	}
-
-	Delegate = MakeUnique<FUsdDiagnosticDelegate>();
-
-	pxr::TfDiagnosticMgr& DiagMgr = pxr::TfDiagnosticMgr::GetInstance();
-	DiagMgr.AddDelegate(Delegate.Get());
-#endif	  // USE_USD_SDK
 }
 
-void UnrealUSDWrapper::ClearDiagnosticDelegate() TSAN_SAFE
+// Deprecated
+void UnrealUSDWrapper::ClearDiagnosticDelegate()
 {
-#if USE_USD_SDK
-	if (!Delegate.IsValid())
-	{
-		return;
-	}
-
-	pxr::TfDiagnosticMgr& DiagMgr = pxr::TfDiagnosticMgr::GetInstance();
-	DiagMgr.RemoveDelegate(Delegate.Get());
-
-	Delegate = nullptr;
-#endif	  // USE_USD_SDK
 }
 
 class FUnrealUSDWrapperModule : public IUnrealUSDWrapperModule
 {
-	FDelegateHandle OnProjectSettingsChangedHandle;
-
 public:
 	virtual void StartupModule() override TSAN_SAFE
 	{
 #if USE_USD_SDK
 		LLM_SCOPE_BYTAG(Usd);
 
-		UUsdProjectSettings* ProjectSettings = GetMutableDefault<UUsdProjectSettings>();
-		if (ProjectSettings->bLogUsdSdkErrors)
-		{
-			UnrealUSDWrapper::SetupDiagnosticDelegate();
-		}
-
-#if WITH_EDITOR
-		OnProjectSettingsChangedHandle = ProjectSettings->OnSettingChanged().AddLambda(
-			[](UObject* SettingsObject, struct FPropertyChangedEvent& PropertyChangedEvent)
-			{
-				UUsdProjectSettings* UsdSettings = Cast<UUsdProjectSettings>(SettingsObject);
-				if (!UsdSettings)
-				{
-					return;
-				}
-
-				if (PropertyChangedEvent.GetPropertyName() == GET_MEMBER_NAME_CHECKED(UUsdProjectSettings, bLogUsdSdkErrors))
-				{
-					if (UsdSettings->bLogUsdSdkErrors)
-					{
-						UnrealUSDWrapper::SetupDiagnosticDelegate();
-						UE_LOG(LogUsd, Log, TEXT("Logging USD messages"));
-					}
-					else
-					{
-						UnrealUSDWrapper::ClearDiagnosticDelegate();
-						UE_LOG(LogUsd, Log, TEXT("Not logging USD messages"));
-					}
-				}
-			}
-		);
-#endif	  // WITH_EDITOR
-
 		// Path to USD base plugins
 		FString UsdPluginsPath = FPaths::Combine(TEXT(".."), TEXT("ThirdParty"), TEXT("USD"), TEXT("UsdResources"));
 		UsdPluginsPath = FPaths::ConvertRelativePathToFull(UsdPluginsPath);
@@ -1216,6 +1109,8 @@ public:
 		// Have to do this in USDClasses as we need the Json module, which is RTTI == false
 		IUsdClassesModule::UpdatePlugInfoFiles(UsdPluginsPath, TargetDllFolder);
 
+		UUsdProjectSettings* ProjectSettings = GetMutableDefault<UUsdProjectSettings>();
+
 		// Combine our current plugins with any additional USD plugins the user may have set.
 		TArray<FString> PluginDirectories;
 		PluginDirectories.Add(UsdPluginsPath);
@@ -1295,19 +1190,6 @@ public:
 		UsdUnreal::MaterialUtils::UnregisterRenderContext(UnrealIdentifiers::UniversalRenderContext);
 		UsdUnreal::MaterialUtils::UnregisterRenderContext(UnrealIdentifiers::UnrealRenderContext);
 
-#if WITH_EDITOR
-		// We can't query default objects during engine exit
-		if (UObjectInitialized())
-		{
-			UUsdProjectSettings* ProjectSettings = GetMutableDefault<UUsdProjectSettings>();
-			if (ProjectSettings)
-			{
-				ProjectSettings->OnSettingChanged().Remove(OnProjectSettingsChangedHandle);
-			}
-		}
-#endif	  // WITH_EDITOR
-
-		UnrealUSDWrapper::ClearDiagnosticDelegate();
 		FUsdMemoryManager::Shutdown();
 	}
 };
diff --git a/Engine/Plugins/Runtime/USDCore/Source/UnrealUSDWrapper/Public/UnrealUSDWrapper.h b/Engine/Plugins/Runtime/USDCore/Source/UnrealUSDWrapper/Public/UnrealUSDWrapper.h
--- a/Engine/Plugins/Runtime/USDCore/Source/UnrealUSDWrapper/Public/UnrealUSDWrapper.h
+++ b/Engine/Plugins/Runtime/USDCore/Source/UnrealUSDWrapper/Public/UnrealUSDWrapper.h
@@ -56,7 +56,6 @@ PXR_NAMESPACE_CLOSE_SCOPE
 #endif	  // #if USE_USD_SDK
 
 class IUsdPrim;
-class FUsdDiagnosticDelegate;
 
 namespace UE
 {
@@ -372,13 +371,12 @@ public:
 	UNREALUSDWRAPPER_API static void SetDefaultResolverDefaultSearchPath(const TArray<FDirectoryPath>& SearchPath);
 
 	/** Starts listening to error/warning/log messages emitted by USD */
+	UE_DEPRECATED(5.6, "Use RegisterDiagnosticDelegate from USDErrorUtils.h.")
 	UNREALUSDWRAPPER_API static void SetupDiagnosticDelegate();
 
 	/** Stops listening to error/warning/log messages emitted by USD */
+	UE_DEPRECATED(5.6, "Use UnregisterDiagnosticDelegate from USDErrorUtils.h.")
 	UNREALUSDWRAPPER_API static void ClearDiagnosticDelegate();
-
-private:
-	static TUniquePtr<FUsdDiagnosticDelegate> Delegate;
 };
 
 class IUsdPrim
diff --git a/Engine/Plugins/Runtime/USDCore/Source/UnrealUSDWrapper/UnrealUSDWrapper.Build.cs b/Engine/Plugins/Runtime/USDCore/Source/UnrealUSDWrapper/UnrealUSDWrapper.Build.cs
--- a/Engine/Plugins/Runtime/USDCore/Source/UnrealUSDWrapper/UnrealUSDWrapper.Build.cs
+++ b/Engine/Plugins/Runtime/USDCore/Source/UnrealUSDWrapper/UnrealUSDWrapper.Build.cs
@@ -24,16 +24,6 @@ namespace UnrealBuildTool.Rules
 				}
 			);
 
-			if (Target.bBuildEditor)
-			{
-				PrivateDependencyModuleNames.AddRange(
-					new string[]
-					{
-						"DeveloperSettings", // So that we can un/register the diagnostic delegate when the project settings change
-					}
-				);
-			}
-
 			// Temporarily disabled runtime USD support until Mac and Linux dynamic linking issues are resolved
 			if (EnableUsdSdk(Target) && (Target.Type == TargetType.Editor || Target.Platform == UnrealTargetPlatform.Win64))
 			{
