USD: Tweak warning text about error utils refactoring.

diff --git a/Engine/Plugins/Runtime/USDCore/Source/USDUtilities/Private/USDErrorUtils.cpp b/Engine/Plugins/Runtime/USDCore/Source/USDUtilities/Private/USDErrorUtils.cpp
--- a/Engine/Plugins/Runtime/USDCore/Source/USDUtilities/Private/USDErrorUtils.cpp
+++ b/Engine/Plugins/Runtime/USDCore/Source/USDUtilities/Private/USDErrorUtils.cpp
@@ -209,7 +209,7 @@ namespace UE::USDErrorUtils::Private
 					EMessageSeverity::Info,
 					LOCTEXT(
 						"SkippedMEssages",
-						"Some similar log messages were skipped during the previous USD operation. You can disable this behavior by setting the USD.CollapseLogMessages cvar to false."
+						"Some similar log messages were skipped during the previous USD operation. You can disable this behavior by unchecking 'Optimize Usd Log' on the Unreal project settings window."
 					)
 				);
 				SendMessageToOutputLog(*SkippedMessage);
diff --git a/Engine/Plugins/Runtime/USDCore/Source/USDUtilities/Public/USDErrorUtils.h b/Engine/Plugins/Runtime/USDCore/Source/USDUtilities/Public/USDErrorUtils.h
--- a/Engine/Plugins/Runtime/USDCore/Source/USDUtilities/Public/USDErrorUtils.h
+++ b/Engine/Plugins/Runtime/USDCore/Source/USDUtilities/Public/USDErrorUtils.h
@@ -101,7 +101,7 @@ public:
 /**
  * Begins a scope where all logged messages (user and non-user facing, USD SDK errors and even USD error mark messages) are collected and deduplicated.
  *
- * The deduplication only happens if the USD.CollapseLogMessages cvar is enabled (default).
+ * The deduplication only happens if the bOptimizeUsdLog project setting is enabled (default).
  *
  * If the USD.UseMessageLog cvar is enabled (default), it will add the user-facing messages (and the USD SDK and error mark messages) to the Message Log,
  * (whether deduplication is enabled or disabled), also displaying a toast letting the user know that those messages can be viewed there.
