USDCore: enable OpenUSD SDK when building with clang on Windows

diff --git a/Engine/Plugins/Runtime/USDCore/Source/UnrealUSDWrapper/UnrealUSDWrapper.Build.cs b/Engine/Plugins/Runtime/USDCore/Source/UnrealUSDWrapper/UnrealUSDWrapper.Build.cs
--- a/Engine/Plugins/Runtime/USDCore/Source/UnrealUSDWrapper/UnrealUSDWrapper.Build.cs
+++ b/Engine/Plugins/Runtime/USDCore/Source/UnrealUSDWrapper/UnrealUSDWrapper.Build.cs
@@ -172,7 +172,6 @@ namespace UnrealBuildTool.Rules
 			// USD SDK has been built against Python 3 and won't launch if the editor is using Python 2
 
 			bool bEnableUsdSdk = (
-				Target.WindowsPlatform.Compiler != WindowsCompiler.Clang &&
 				Target.WindowsPlatform.Compiler != WindowsCompiler.Intel &&
 				Target.StaticAnalyzer == StaticAnalyzer.None
 			);
