USD: Analytics tweaks.

Send asset counts when opening the stage;
Send whether exported static mesh had Nanite enabled;

diff --git a/Engine/Plugins/Importers/USDImporter/Source/USDExporter/Private/StaticMeshExporterUSD.cpp b/Engine/Plugins/Importers/USDImporter/Source/USDExporter/Private/StaticMeshExporterUSD.cpp
--- a/Engine/Plugins/Importers/USDImporter/Source/USDExporter/Private/StaticMeshExporterUSD.cpp
+++ b/Engine/Plugins/Importers/USDImporter/Source/USDExporter/Private/StaticMeshExporterUSD.cpp
@@ -32,7 +32,7 @@
 namespace UE::StaticMeshExporterUSD::Private
 {
 	void SendAnalytics(
-		UObject* Asset,
+		UStaticMesh* Asset,
 		UStaticMeshExporterUSDOptions* Options,
 		bool bAutomated,
 		double ElapsedSeconds,
@@ -49,6 +49,7 @@ namespace UE::StaticMeshExporterUSD::Private
 
 		TArray<FAnalyticsEventAttribute> EventAttributes;
 		EventAttributes.Emplace(TEXT("AssetType"), ClassName);
+		EventAttributes.Emplace(TEXT("IsNaniteEnabled"), Asset->IsNaniteEnabled());
 
 		if (Options)
 		{
@@ -390,7 +391,7 @@ bool UStaticMeshExporterUsd::ExportBinary(UObject* Object, const TCHAR* Type, FA
 		FString Extension = FPaths::GetExtension(UExporter::CurrentFilename);
 
 		UE::StaticMeshExporterUSD::Private::SendAnalytics(
-			Object,
+			StaticMesh,
 			Options,
 			bAutomated,
 			ElapsedSeconds,
diff --git a/Engine/Plugins/Importers/USDImporter/Source/USDStage/Private/USDStageActor.cpp b/Engine/Plugins/Importers/USDImporter/Source/USDStage/Private/USDStageActor.cpp
--- a/Engine/Plugins/Importers/USDImporter/Source/USDStage/Private/USDStageActor.cpp
+++ b/Engine/Plugins/Importers/USDImporter/Source/USDStage/Private/USDStageActor.cpp
@@ -566,7 +566,13 @@ struct FUsdStageActorImpl
 		PrimSceneComponent->SetHiddenInGame(bNewHidden, bPropagateToChildren);
 	}
 
-	static void SendAnalytics(AUsdStageActor* StageActor, double ElapsedSeconds, double NumberOfFrames, const FString& Extension)
+	static void SendAnalytics(
+		AUsdStageActor* StageActor,
+		double ElapsedSeconds,
+		double NumberOfFrames,
+		const FString& Extension,
+		const TSet<FSoftObjectPath>& ActiveAssetPaths
+	)
 	{
 		if (!StageActor)
 		{
@@ -594,6 +600,15 @@ struct FUsdStageActorImpl
 
 			UsdUtils::AddAnalyticsAttributes(StageActor->MetadataOptions, EventAttributes);
 
+			TSet<UObject*> ActiveAssets;
+			ActiveAssets.Reserve(ActiveAssetPaths.Num());
+			for (const FSoftObjectPath& ActivePath : ActiveAssetPaths)
+			{
+				ActiveAssets.Add(ActivePath.TryLoad());
+			}
+			ActiveAssets.Remove(nullptr);
+			IUsdClassesModule::AddAssetCountAttributes(ActiveAssets, EventAttributes);
+
 			const bool bAutomated = false;
 			IUsdClassesModule::SendAnalytics(MoveTemp(EventAttributes), TEXT("Open"), bAutomated, ElapsedSeconds, NumberOfFrames, Extension);
 		}
@@ -3291,6 +3306,7 @@ void AUsdStageActor::LoadUsdStage(bool bOpenIfNeeded)
 	FScopedUsdMessageLog ScopedMessageLog;
 
 	UUsdAssetCache3::FUsdScopedReferencer ScopedReferencer{AssetCache, this};
+	AssetCache->MarkAssetsAsStale();
 
 	UE::FUsdStage StageToLoad = GetUsdStage();
 
@@ -3364,7 +3380,14 @@ void AUsdStageActor::LoadUsdStage(bool bOpenIfNeeded)
 	USD_LOG_INFO(TEXT("%s %s in [%d min %.3f s]"), TEXT("Stage loaded"), *FPaths::GetBaseFilename(RootLayer.FilePath), ElapsedMin, ElapsedSeconds);
 
 #if USE_USD_SDK
-	FUsdStageActorImpl::SendAnalytics(this, ElapsedSeconds, UsdUtils::GetUsdStageNumFrames(StageToLoad), FPaths::GetExtension(RootLayer.FilePath));
+	TSet<FSoftObjectPath> ActiveAssetPaths = AssetCache->GetActiveAssets();
+	FUsdStageActorImpl::SendAnalytics(
+		this,
+		ElapsedSeconds,
+		UsdUtils::GetUsdStageNumFrames(StageToLoad),
+		FPaths::GetExtension(RootLayer.FilePath),
+		ActiveAssetPaths
+	);
 #endif	  // #if USE_USD_SDK
 }
 
diff --git a/Engine/Plugins/Importers/USDImporter/Source/USDStageImporter/Private/USDStageImporter.cpp b/Engine/Plugins/Importers/USDImporter/Source/USDStageImporter/Private/USDStageImporter.cpp
--- a/Engine/Plugins/Importers/USDImporter/Source/USDStageImporter/Private/USDStageImporter.cpp
+++ b/Engine/Plugins/Importers/USDImporter/Source/USDStageImporter/Private/USDStageImporter.cpp
@@ -1550,86 +1550,7 @@ namespace UE::USDStageImporter::Private
 				UsdUtils::AddAnalyticsAttributes(*ImportContext.ImportOptions, EventAttributes);
 			}
 
-			int32 NumStaticMeshes = 0;
-			int32 NumSkeletalMeshes = 0;
-			int32 NumMaterials = 0;
-			int32 NumAnimSequences = 0;
-			int32 NumLevelSequences = 0;
-			int32 NumTextures = 0;
-			int32 NumGeometryCaches = 0;
-			int32 NumGroomAssets = 0;
-			int32 NumGroomBindings = 0;
-			int32 NumGroomCaches = 0;
-			int32 NumSparseVolumeTextures = 0;
-			int32 NumSounds = 0;
-			for (UObject* ImportedAsset : ImportedAssets)
-			{
-				if (!ImportedAsset)
-				{
-					continue;
-				}
-
-				if (ImportedAsset->IsA<UStaticMesh>())
-				{
-					++NumStaticMeshes;
-				}
-				else if (ImportedAsset->IsA<USkeletalMesh>())
-				{
-					++NumSkeletalMeshes;
-				}
-				else if (ImportedAsset->IsA<UMaterialInterface>())
-				{
-					++NumMaterials;
-				}
-				else if (ImportedAsset->IsA<UAnimSequence>())
-				{
-					++NumAnimSequences;
-				}
-				else if (ImportedAsset->IsA<ULevelSequence>())
-				{
-					++NumLevelSequences;
-				}
-				else if (ImportedAsset->IsA<UTexture>())
-				{
-					++NumTextures;
-				}
-				else if (ImportedAsset->IsA<UGeometryCache>())
-				{
-					++NumGeometryCaches;
-				}
-				else if (ImportedAsset->IsA<UGroomAsset>())
-				{
-					++NumGroomAssets;
-				}
-				else if (ImportedAsset->IsA<UGroomBindingAsset>())
-				{
-					++NumGroomBindings;
-				}
-				else if (ImportedAsset->IsA<UGroomCache>())
-				{
-					++NumGroomCaches;
-				}
-				else if (ImportedAsset->IsA<USparseVolumeTexture>())
-				{
-					++NumSparseVolumeTextures;
-				}
-				else if (ImportedAsset->IsA<USoundBase>())
-				{
-					++NumSounds;
-				}
-			}
-			EventAttributes.Emplace(TEXT("NumStaticMeshes"), LexToString(NumStaticMeshes));
-			EventAttributes.Emplace(TEXT("NumSkeletalMeshes"), LexToString(NumSkeletalMeshes));
-			EventAttributes.Emplace(TEXT("NumMaterials"), LexToString(NumMaterials));
-			EventAttributes.Emplace(TEXT("NumAnimSequences"), LexToString(NumAnimSequences));
-			EventAttributes.Emplace(TEXT("NumLevelSequences"), LexToString(NumLevelSequences));
-			EventAttributes.Emplace(TEXT("NumTextures"), LexToString(NumTextures));
-			EventAttributes.Emplace(TEXT("NumGeometryCaches"), LexToString(NumGeometryCaches));
-			EventAttributes.Emplace(TEXT("NumGroomAssets"), LexToString(NumGroomAssets));
-			EventAttributes.Emplace(TEXT("NumGroomBindings"), LexToString(NumGroomBindings));
-			EventAttributes.Emplace(TEXT("NumGroomCaches"), LexToString(NumGroomCaches));
-			EventAttributes.Emplace(TEXT("NumSparseVolumeTextures"), LexToString(NumSparseVolumeTextures));
-			EventAttributes.Emplace(TEXT("NumSounds"), LexToString(NumSounds));
+			IUsdClassesModule::AddAssetCountAttributes(ImportedAssets, EventAttributes);
 
 			FString RootLayerIdentifier = ImportContext.FilePath;
 			if (ImportContext.Stage)
diff --git a/Engine/Plugins/Runtime/USDCore/Source/USDClasses/Private/USDClassesModule.cpp b/Engine/Plugins/Runtime/USDCore/Source/USDClasses/Private/USDClassesModule.cpp
--- a/Engine/Plugins/Runtime/USDCore/Source/USDClasses/Private/USDClassesModule.cpp
+++ b/Engine/Plugins/Runtime/USDCore/Source/USDClasses/Private/USDClassesModule.cpp
@@ -25,6 +25,7 @@
 #include "GroomBindingAsset.h"
 #include "GroomCache.h"
 #include "HAL/FileManager.h"
+#include "LevelSequence.h"
 #include "Materials/Material.h"
 #include "Materials/MaterialInstance.h"
 #include "Materials/MaterialInterface.h"
@@ -272,6 +273,92 @@ void IUsdClassesModule::SendAnalytics(TArray<FAnalyticsEventAttribute>&& InAttri
 	}
 }
 
+void IUsdClassesModule::AddAssetCountAttributes(const TSet<UObject*>& Assets, TArray<FAnalyticsEventAttribute>& InOutAttributes)
+{
+	int32 NumStaticMeshes = 0;
+	int32 NumSkeletalMeshes = 0;
+	int32 NumMaterials = 0;
+	int32 NumAnimSequences = 0;
+	int32 NumLevelSequences = 0;
+	int32 NumTextures = 0;
+	int32 NumGeometryCaches = 0;
+	int32 NumGroomAssets = 0;
+	int32 NumGroomBindings = 0;
+	int32 NumGroomCaches = 0;
+	int32 NumSparseVolumeTextures = 0;
+	int32 NumSounds = 0;
+
+	for (UObject* ImportedAsset : Assets)
+	{
+		if (!ImportedAsset)
+		{
+			continue;
+		}
+
+		if (ImportedAsset->IsA<UStaticMesh>())
+		{
+			++NumStaticMeshes;
+		}
+		else if (ImportedAsset->IsA<USkeletalMesh>())
+		{
+			++NumSkeletalMeshes;
+		}
+		else if (ImportedAsset->IsA<UMaterialInterface>())
+		{
+			++NumMaterials;
+		}
+		else if (ImportedAsset->IsA<UAnimSequence>())
+		{
+			++NumAnimSequences;
+		}
+		else if (ImportedAsset->IsA<ULevelSequence>())
+		{
+			++NumLevelSequences;
+		}
+		else if (ImportedAsset->IsA<UTexture>())
+		{
+			++NumTextures;
+		}
+		else if (ImportedAsset->IsA<UGeometryCache>())
+		{
+			++NumGeometryCaches;
+		}
+		else if (ImportedAsset->IsA<UGroomAsset>())
+		{
+			++NumGroomAssets;
+		}
+		else if (ImportedAsset->IsA<UGroomBindingAsset>())
+		{
+			++NumGroomBindings;
+		}
+		else if (ImportedAsset->IsA<UGroomCache>())
+		{
+			++NumGroomCaches;
+		}
+		else if (ImportedAsset->IsA<USparseVolumeTexture>())
+		{
+			++NumSparseVolumeTextures;
+		}
+		else if (ImportedAsset->IsA<USoundBase>())
+		{
+			++NumSounds;
+		}
+	}
+
+	InOutAttributes.Emplace(TEXT("NumStaticMeshes"), LexToString(NumStaticMeshes));
+	InOutAttributes.Emplace(TEXT("NumSkeletalMeshes"), LexToString(NumSkeletalMeshes));
+	InOutAttributes.Emplace(TEXT("NumMaterials"), LexToString(NumMaterials));
+	InOutAttributes.Emplace(TEXT("NumAnimSequences"), LexToString(NumAnimSequences));
+	InOutAttributes.Emplace(TEXT("NumLevelSequences"), LexToString(NumLevelSequences));
+	InOutAttributes.Emplace(TEXT("NumTextures"), LexToString(NumTextures));
+	InOutAttributes.Emplace(TEXT("NumGeometryCaches"), LexToString(NumGeometryCaches));
+	InOutAttributes.Emplace(TEXT("NumGroomAssets"), LexToString(NumGroomAssets));
+	InOutAttributes.Emplace(TEXT("NumGroomBindings"), LexToString(NumGroomBindings));
+	InOutAttributes.Emplace(TEXT("NumGroomCaches"), LexToString(NumGroomCaches));
+	InOutAttributes.Emplace(TEXT("NumSparseVolumeTextures"), LexToString(NumSparseVolumeTextures));
+	InOutAttributes.Emplace(TEXT("NumSounds"), LexToString(NumSounds));
+}
+
 void IUsdClassesModule::BlockAnalyticsEvents()
 {
 	UE::USDClasses::Private::AnalyticsAllowedWhenZero += 1;
diff --git a/Engine/Plugins/Runtime/USDCore/Source/USDClasses/Public/USDClassesModule.h b/Engine/Plugins/Runtime/USDCore/Source/USDClasses/Public/USDClassesModule.h
--- a/Engine/Plugins/Runtime/USDCore/Source/USDClasses/Public/USDClassesModule.h
+++ b/Engine/Plugins/Runtime/USDCore/Source/USDClasses/Public/USDClassesModule.h
@@ -37,6 +37,8 @@ public:
 	);
 	USDCLASSES_API static void SendAnalytics(TArray<FAnalyticsEventAttribute>&& InAttributes, const FString& EventName);
 
+	USDCLASSES_API static void AddAssetCountAttributes(const TSet<UObject*>& Assets, TArray<FAnalyticsEventAttribute>& InOutAttributes);
+
 	USDCLASSES_API static void BlockAnalyticsEvents();
 
 	USDCLASSES_API static void ResumeAnalyticsEvents();
