USD: Fix broken skeletal meshes when handling faces with less than 3 vertices.

Introduced with cl 38851077

diff --git a/Engine/Plugins/Runtime/USDCore/Source/USDUtilities/Private/USDSkeletalDataConversion.cpp b/Engine/Plugins/Runtime/USDCore/Source/USDUtilities/Private/USDSkeletalDataConversion.cpp
--- a/Engine/Plugins/Runtime/USDCore/Source/USDUtilities/Private/USDSkeletalDataConversion.cpp
+++ b/Engine/Plugins/Runtime/USDCore/Source/USDUtilities/Private/USDSkeletalDataConversion.cpp
@@ -1803,6 +1803,7 @@ bool UsdToUnreal::ConvertSkinnedMesh(
 		if (NumOriginalFaceVertices < 3)
 		{
 			++NumSkippedPolygons;
+			NumProcessedFaceVertexIndices += NumOriginalFaceVertices;
 			continue;
 		}
 
