USD: Show some feedback when trying to reimport with USDImporter something that was imported via USDInterchange.

diff --git a/Engine/Plugins/Importers/USDImporter/Source/USDStageImporter/Private/USDStageAssetImportFactory.cpp b/Engine/Plugins/Importers/USDImporter/Source/USDStageImporter/Private/USDStageAssetImportFactory.cpp
--- a/Engine/Plugins/Importers/USDImporter/Source/USDStageImporter/Private/USDStageAssetImportFactory.cpp
+++ b/Engine/Plugins/Importers/USDImporter/Source/USDStageImporter/Private/USDStageAssetImportFactory.cpp
@@ -121,7 +121,7 @@ void UUsdStageAssetImportFactory::CleanUp()
 
 bool UUsdStageAssetImportFactory::CanReimport(UObject* Obj, TArray<FString>& OutFilenames)
 {
-	if (UAssetImportData* ImportData = UsdUnreal::ObjectUtils::GetAssetImportData(Obj))
+	if (UAssetImportData* ImportData = UsdUnreal::ObjectUtils::GetBaseAssetImportData(Obj))
 	{
 		const FString FileName = ImportData->GetFirstFilename();
 		const FString FileExtension = FPaths::GetExtension(FileName);
@@ -163,6 +163,8 @@ EReimportResult::Type UUsdStageAssetImportFactory::Reimport(UObject* Obj)
 		return EReimportResult::Failed;
 	}
 
+	FScopedUsdMessageLog ScopedMessageLog;
+
 	FString ReimportFilePath;
 	FString OriginalPrimPath;
 	UObject* ReimportOptions = nullptr;
@@ -187,7 +189,10 @@ EReimportResult::Type UUsdStageAssetImportFactory::Reimport(UObject* Obj)
 	if (ReimportFilePath.IsEmpty() || OriginalPrimPath.IsEmpty())
 	{
 		USD_LOG_USERERROR(FText::Format(
-			LOCTEXT("ReimportErrorNoImportData", "Failed to reimport asset '{0}' as it doesn't seem to have valid USD import data or user data!"),
+			LOCTEXT(
+				"ReimportErrorNoImportData",
+				"Failed to reimport asset '{0}' as it doesn't seem to have valid USD import data or user data! Only assets created by the USD Importer plugin can be reimported via the USD Importer plugin."
+			),
 			FText::FromName(Obj->GetFName())
 		));
 		return EReimportResult::Failed;
@@ -216,8 +221,6 @@ EReimportResult::Type UUsdStageAssetImportFactory::Reimport(UObject* Obj)
 
 	ImportContext.PackagePath = Obj->GetOutermost()->GetPathName();
 
-	FScopedUsdMessageLog ScopedMessageLog;
-
 	UUsdStageImporter* USDImporter = IUsdStageImporterModule::Get().GetImporter();
 	UObject* ReimportedAsset = nullptr;
 	bool bSuccess = USDImporter->ReimportSingleAsset(ImportContext, Obj, OriginalPrimPath, ReimportedAsset);
diff --git a/Engine/Plugins/Runtime/USDCore/Source/USDClasses/Private/USDObjectUtils.cpp b/Engine/Plugins/Runtime/USDCore/Source/USDClasses/Private/USDObjectUtils.cpp
--- a/Engine/Plugins/Runtime/USDCore/Source/USDClasses/Private/USDObjectUtils.cpp
+++ b/Engine/Plugins/Runtime/USDCore/Source/USDClasses/Private/USDObjectUtils.cpp
@@ -23,26 +23,26 @@
 #include "PhysicsEngine/PhysicsAsset.h"
 #include "SparseVolumeTexture/SparseVolumeTexture.h"
 
-UUsdAssetImportData* UsdUnreal::ObjectUtils::GetAssetImportData(UObject* Asset)
+UAssetImportData* UsdUnreal::ObjectUtils::GetBaseAssetImportData(UObject* Asset)
 {
-	UUsdAssetImportData* ImportData = nullptr;
+	UAssetImportData* ImportData = nullptr;
 #if WITH_EDITORONLY_DATA
 	if (UStaticMesh* Mesh = Cast<UStaticMesh>(Asset))
 	{
-		ImportData = Cast<UUsdAssetImportData>(Mesh->AssetImportData);
+		ImportData = Mesh->AssetImportData;
 	}
 	else if (USkeleton* Skeleton = Cast<USkeleton>(Asset))
 	{
 		if (USkeletalMesh* SkelMesh = Skeleton->GetPreviewMesh())
 		{
-			ImportData = Cast<UUsdAssetImportData>(SkelMesh->GetAssetImportData());
+			ImportData = SkelMesh->GetAssetImportData();
 		}
 	}
 	else if (UPhysicsAsset* PhysicsAsset = Cast<UPhysicsAsset>(Asset))
 	{
 		if (USkeletalMesh* SkelMesh = PhysicsAsset->GetPreviewMesh())
 		{
-			ImportData = Cast<UUsdAssetImportData>(SkelMesh->GetAssetImportData());
+			ImportData = SkelMesh->GetAssetImportData();
 		}
 	}
 	else if (UAnimBlueprint* AnimBP = Cast<UAnimBlueprint>(Asset))
@@ -53,47 +53,52 @@ UUsdAssetImportData* UsdUnreal::ObjectUtils::GetAssetImportData(UObject* Asset)
 		{
 			if (USkeletalMesh* SkelMesh = AnimBPSkeleton->GetPreviewMesh())
 			{
-				ImportData = Cast<UUsdAssetImportData>(SkelMesh->GetAssetImportData());
+				ImportData = SkelMesh->GetAssetImportData();
 			}
 		}
 	}
 	else if (USkeletalMesh* SkelMesh = Cast<USkeletalMesh>(Asset))
 	{
-		ImportData = Cast<UUsdAssetImportData>(SkelMesh->GetAssetImportData());
+		ImportData = SkelMesh->GetAssetImportData();
 	}
 	else if (UAnimSequence* SkelAnim = Cast<UAnimSequence>(Asset))
 	{
-		ImportData = Cast<UUsdAssetImportData>(SkelAnim->AssetImportData);
+		ImportData = SkelAnim->AssetImportData;
 	}
 	else if (UMaterialInterface* Material = Cast<UMaterialInterface>(Asset))
 	{
-		ImportData = Cast<UUsdAssetImportData>(Material->AssetImportData);
+		ImportData = Material->AssetImportData;
 	}
 	else if (UTexture* Texture = Cast<UTexture>(Asset))
 	{
-		ImportData = Cast<UUsdAssetImportData>(Texture->AssetImportData);
+		ImportData = Texture->AssetImportData;
 	}
 	else if (UGeometryCache* GeometryCache = Cast<UGeometryCache>(Asset))
 	{
-		ImportData = Cast<UUsdAssetImportData>(GeometryCache->AssetImportData);
+		ImportData = GeometryCache->AssetImportData;
 	}
 	else if (UGroomAsset* Groom = Cast<UGroomAsset>(Asset))
 	{
-		ImportData = Cast<UUsdAssetImportData>(Groom->AssetImportData);
+		ImportData = Groom->AssetImportData;
 	}
 	else if (UGroomCache* GroomCache = Cast<UGroomCache>(Asset))
 	{
-		ImportData = Cast<UUsdAssetImportData>(GroomCache->AssetImportData);
+		ImportData = GroomCache->AssetImportData;
 	}
 	else if (UStreamableSparseVolumeTexture* SparseVolumeTexture = Cast<UStreamableSparseVolumeTexture>(Asset))
 	{
-		ImportData = Cast<UUsdAssetImportData>(SparseVolumeTexture->AssetImportData);
+		ImportData = SparseVolumeTexture->AssetImportData;
 	}
 
 #endif
 	return ImportData;
 }
 
+UUsdAssetImportData* UsdUnreal::ObjectUtils::GetAssetImportData(UObject* Asset)
+{
+	return Cast<UUsdAssetImportData>(GetBaseAssetImportData(Asset));
+}
+
 void UsdUnreal::ObjectUtils::SetAssetImportData(UObject* Asset, UAssetImportData* ImportData)
 {
 	if (!Asset)
diff --git a/Engine/Plugins/Runtime/USDCore/Source/USDClasses/Public/USDObjectUtils.h b/Engine/Plugins/Runtime/USDCore/Source/USDClasses/Public/USDObjectUtils.h
--- a/Engine/Plugins/Runtime/USDCore/Source/USDClasses/Public/USDObjectUtils.h
+++ b/Engine/Plugins/Runtime/USDCore/Source/USDClasses/Public/USDObjectUtils.h
@@ -18,6 +18,7 @@ namespace UsdUnreal::ObjectUtils
 	 * we'll actually check it's preview mesh instead (since Skeletons don't have AssetImportData). The setter won't do
 	 * anything if you try setting asset import data on e.g. a Skeleton, on the other hand.
 	 */
+	USDCLASSES_API UAssetImportData* GetBaseAssetImportData(UObject* Asset);
 	USDCLASSES_API UUsdAssetImportData* GetAssetImportData(UObject* Asset);
 	USDCLASSES_API void SetAssetImportData(UObject* Asset, UAssetImportData* ImportData);
 
