USD: Fix import of skeletal animations with blend shapes not flagging the curves as morph targets.

diff --git a/Engine/Plugins/Runtime/USDCore/Source/USDUtilities/Private/USDSkeletalDataConversion.cpp b/Engine/Plugins/Runtime/USDCore/Source/USDUtilities/Private/USDSkeletalDataConversion.cpp
--- a/Engine/Plugins/Runtime/USDCore/Source/USDUtilities/Private/USDSkeletalDataConversion.cpp
+++ b/Engine/Plugins/Runtime/USDCore/Source/USDUtilities/Private/USDSkeletalDataConversion.cpp
@@ -2967,6 +2967,10 @@ bool UsdToUnreal::ConvertSkeletalImportData(
 			const bool bMaterialCurve = false;
 			const bool bMorphTargetCurve = true;
 			Skeleton->AccumulateCurveMetaData(CurveName, bMaterialCurve, bMorphTargetCurve);
+
+			// Ensure we have a morph flag set
+			FCurveMetaData* CurveMetaData = Skeleton->GetCurveMetaData(CurveName);
+			CurveMetaData->Type.bMorphtarget = true;
 		}
 		else
 		{
