USD: Fix off-by-one error when exporting AnimationSequences to USD.

diff --git a/Engine/Plugins/Runtime/USDCore/Source/USDUtilities/Private/USDSkeletalDataConversion.cpp b/Engine/Plugins/Runtime/USDCore/Source/USDUtilities/Private/USDSkeletalDataConversion.cpp
--- a/Engine/Plugins/Runtime/USDCore/Source/USDUtilities/Private/USDSkeletalDataConversion.cpp
+++ b/Engine/Plugins/Runtime/USDCore/Source/USDUtilities/Private/USDSkeletalDataConversion.cpp
@@ -3877,7 +3877,10 @@ bool UnrealToUsd::ConvertAnimSequence(UAnimSequence* AnimSequence, pxr::UsdPrim&
 	const FReferenceSkeleton& RefSkeleton = SkeletalMesh->GetRefSkeleton();
 	const int32 NumBones = RefSkeleton.GetRefBoneInfo().Num();
 	const double TimeCodesPerSecond = SkelAnimPrim.GetStage()->GetTimeCodesPerSecond();
-	const int32 NumTimeCodes = AnimSequence->GetPlayLength() * TimeCodesPerSecond;
+
+	// The +1 is because 1s length at 1 tcps implies we want two time codes: One at zero, and one at 1s
+	// The ceil to make sure we never clip the animation in case it doesn't end exactly on an exported time code value
+	const int32 NumTimeCodes = FMath::CeilToInt(AnimSequence->GetPlayLength() * TimeCodesPerSecond) + 1;
 
 	if (NumBones <= 0)
 	{
