USD: More analytics tweaks.

- Remove duplicate EndTimeCode attribute from level export events;
- Rename attribute to LevelSequenceSelectionOnly to avoid conflict with SelectionOnly, which refers to the level export;
- Pick asset class name from a fixed set of names, as opposed to using whatever the derived class name is;
- Prevent separate Level export event from being fired when exporting LevelSequences;
- Always collect Level export attributes when exporting LevelSequences.

diff --git a/Engine/Plugins/Importers/USDImporter/Source/USDExporter/Private/AnimSequenceExporterUSD.cpp b/Engine/Plugins/Importers/USDImporter/Source/USDExporter/Private/AnimSequenceExporterUSD.cpp
--- a/Engine/Plugins/Importers/USDImporter/Source/USDExporter/Private/AnimSequenceExporterUSD.cpp
+++ b/Engine/Plugins/Importers/USDImporter/Source/USDExporter/Private/AnimSequenceExporterUSD.cpp
@@ -45,7 +45,7 @@ namespace UE::AnimSequenceExporterUSD::Private
 			return;
 		}
 
-		FString ClassName = Asset->GetClass()->GetName();
+		const FString ClassName = IUsdClassesModule::GetClassNameForAnalytics(Asset);
 
 		TArray<FAnalyticsEventAttribute> EventAttributes;
 		EventAttributes.Emplace(TEXT("AssetType"), ClassName);
diff --git a/Engine/Plugins/Importers/USDImporter/Source/USDExporter/Private/GeometryCacheExporterUSD.cpp b/Engine/Plugins/Importers/USDImporter/Source/USDExporter/Private/GeometryCacheExporterUSD.cpp
--- a/Engine/Plugins/Importers/USDImporter/Source/USDExporter/Private/GeometryCacheExporterUSD.cpp
+++ b/Engine/Plugins/Importers/USDImporter/Source/USDExporter/Private/GeometryCacheExporterUSD.cpp
@@ -43,7 +43,7 @@ namespace UE::GeometryCacheExporterUSD::Private
 			return;
 		}
 
-		FString ClassName = Asset->GetClass()->GetName();
+		const FString ClassName = IUsdClassesModule::GetClassNameForAnalytics(Asset);
 
 		TArray<FAnalyticsEventAttribute> EventAttributes;
 		EventAttributes.Emplace(TEXT("AssetType"), ClassName);
diff --git a/Engine/Plugins/Importers/USDImporter/Source/USDExporter/Private/LevelExporterUSDOptions.cpp b/Engine/Plugins/Importers/USDImporter/Source/USDExporter/Private/LevelExporterUSDOptions.cpp
--- a/Engine/Plugins/Importers/USDImporter/Source/USDExporter/Private/LevelExporterUSDOptions.cpp
+++ b/Engine/Plugins/Importers/USDImporter/Source/USDExporter/Private/LevelExporterUSDOptions.cpp
@@ -36,7 +36,6 @@ void UsdUtils::AddAnalyticsAttributes(const ULevelExporterUSDOptions& Options, T
 	InOutAttributes.Emplace(TEXT("StartTimeCode"), LexToString(Options.StartTimeCode));
 	InOutAttributes.Emplace(TEXT("EndTimeCode"), LexToString(Options.EndTimeCode));
 	UsdUtils::AddAnalyticsAttributes(Options.Inner, InOutAttributes);
-	InOutAttributes.Emplace(TEXT("EndTimeCode"), LexToString(Options.EndTimeCode));
 	InOutAttributes.Emplace(TEXT("ReExportIdenticalLevelsAndSequences"), Options.bReExportIdenticalLevelsAndSequences);
 	InOutAttributes.Emplace(TEXT("ReExportIdenticalAssets"), Options.bReExportIdenticalAssets);
 }
diff --git a/Engine/Plugins/Importers/USDImporter/Source/USDExporter/Private/LevelSequenceExporterUSD.cpp b/Engine/Plugins/Importers/USDImporter/Source/USDExporter/Private/LevelSequenceExporterUSD.cpp
--- a/Engine/Plugins/Importers/USDImporter/Source/USDExporter/Private/LevelSequenceExporterUSD.cpp
+++ b/Engine/Plugins/Importers/USDImporter/Source/USDExporter/Private/LevelSequenceExporterUSD.cpp
@@ -76,7 +76,8 @@ namespace UE::LevelSequenceExporterUSD::Private
 		const TArray<UE::FUsdStage>& ExportedStages,
 		bool bAutomated,
 		double ElapsedSeconds,
-		const FString& Extension
+		const FString& Extension,
+		TArray<FAnalyticsEventAttribute> EventAttributes
 	)
 	{
 		if (!LevelSequence || ExportedStages.Num() == 0 || !FEngineAnalytics::IsAvailable())
@@ -84,9 +85,8 @@ namespace UE::LevelSequenceExporterUSD::Private
 			return;
 		}
 
-		FString ClassName = LevelSequence->GetClass()->GetName();
+		const FString ClassName = IUsdClassesModule::GetClassNameForAnalytics(LevelSequence);
 
-		TArray<FAnalyticsEventAttribute> EventAttributes;
 		EventAttributes.Emplace(TEXT("AssetType"), ClassName);
 		EventAttributes.Emplace(TEXT("NumExportedLevelSequenceLayers"), ExportedStages.Num());
 
@@ -97,8 +97,21 @@ namespace UE::LevelSequenceExporterUSD::Private
 			UsdUtils::AddAnalyticsAttributes(*Options, EventAttributes);
 		}
 
+		// Go through a map because we'll get some overlapping attributes here due to reading analytics
+		// produced by the level import. Whenever they conflict, we want the LevelSequence values.
+		// This will naturally happen here anyway because we add all the LevelSequence attributes afterward,
+		// which will overwrite in the map
+		TMap<FString, FAnalyticsEventAttribute> AttributesByName;
+		AttributesByName.Reserve(EventAttributes.Num());
+		for (const FAnalyticsEventAttribute& Attribute : EventAttributes)
+		{
+			AttributesByName.FindOrAdd(Attribute.GetName()) = Attribute;
+		}
+		TArray<FAnalyticsEventAttribute> UniqueAttributes;
+		AttributesByName.GenerateValueArray(UniqueAttributes);
+
 		IUsdClassesModule::SendAnalytics(
-			MoveTemp(EventAttributes),
+			MoveTemp(UniqueAttributes),
 			FString::Printf(TEXT("Export.%s"), *ClassName),
 			bAutomated,
 			ElapsedSeconds,
@@ -1918,6 +1931,8 @@ bool ULevelSequenceExporterUsd::ExportBinary(
 		}
 	}
 
+	TArray<FAnalyticsEventAttribute> AnalyticsAttributes;
+
 	// Export level if we need to
 	if (Options && Options->bExportLevel)
 	{
@@ -1953,7 +1968,19 @@ bool ULevelSequenceExporterUsd::ExportBinary(
 			LevelExportTask->bWriteEmptyFiles = false;
 			LevelExportTask->bAutomated = true;	   // Pretend this is an automated task so it doesn't pop the options dialog
 
-			UExporter::RunAssetExportTask(LevelExportTask);
+			// We don't want the level export to send a separate event
+			IUsdClassesModule::BlockAnalyticsEvents();
+			{
+				UExporter::RunAssetExportTask(LevelExportTask);
+			}
+			IUsdClassesModule::ResumeAnalyticsEvents();
+
+			// Grab the analytics that the level exporter produced, because that also includes some asset counts that would be
+			// awkward to collect from here
+			if (const TArray<FAnalyticsEventAttribute>* LevelAnalytics = IUsdClassesModule::GetAccumulatedAnalytics().Find(TEXT("Export.World")))
+			{
+				AnalyticsAttributes = *LevelAnalytics;
+			}
 		}
 	}
 
@@ -1977,7 +2004,15 @@ bool ULevelSequenceExporterUsd::ExportBinary(
 		double ElapsedSeconds = FPlatformTime::ToSeconds64(FPlatformTime::Cycles64() - StartTime);
 		FString Extension = FPaths::GetExtension(UExporter::CurrentFilename);
 
-		UE::LevelSequenceExporterUSD::Private::SendAnalytics(LevelSequence, Options, ExportedStages, bAutomated, ElapsedSeconds, Extension);
+		UE::LevelSequenceExporterUSD::Private::SendAnalytics(
+			LevelSequence,
+			Options,
+			ExportedStages,
+			bAutomated,
+			ElapsedSeconds,
+			Extension,
+			AnalyticsAttributes
+		);
 	}
 
 	return true;
diff --git a/Engine/Plugins/Importers/USDImporter/Source/USDExporter/Private/LevelSequenceExporterUSDOptions.cpp b/Engine/Plugins/Importers/USDImporter/Source/USDExporter/Private/LevelSequenceExporterUSDOptions.cpp
--- a/Engine/Plugins/Importers/USDImporter/Source/USDExporter/Private/LevelSequenceExporterUSDOptions.cpp
+++ b/Engine/Plugins/Importers/USDImporter/Source/USDExporter/Private/LevelSequenceExporterUSDOptions.cpp
@@ -17,7 +17,7 @@ void UsdUtils::AddAnalyticsAttributes(const ULevelSequenceExporterUsdOptions& Op
 		InOutAttributes.Emplace(TEXT("StartFrame"), LexToString(Options.StartFrame));
 		InOutAttributes.Emplace(TEXT("EndFrame"), LexToString(Options.EndFrame));
 	}
-	InOutAttributes.Emplace(TEXT("SelectionOnly"), Options.bSelectionOnly);
+	InOutAttributes.Emplace(TEXT("LevelSequenceSelectionOnly"), Options.bSelectionOnly);
 	InOutAttributes.Emplace(TEXT("ExportSubsequencesAsLayers"), Options.bExportSubsequencesAsLayers);
 	InOutAttributes.Emplace(TEXT("ExportLevel"), Options.bExportLevel);
 	if (Options.bExportLevel)
diff --git a/Engine/Plugins/Importers/USDImporter/Source/USDExporter/Private/MaterialExporterUSD.cpp b/Engine/Plugins/Importers/USDImporter/Source/USDExporter/Private/MaterialExporterUSD.cpp
--- a/Engine/Plugins/Importers/USDImporter/Source/USDExporter/Private/MaterialExporterUSD.cpp
+++ b/Engine/Plugins/Importers/USDImporter/Source/USDExporter/Private/MaterialExporterUSD.cpp
@@ -55,7 +55,7 @@ namespace UE::MaterialExporterUSD::Private
 			return;
 		}
 
-		FString ClassName = Material.GetClass()->GetName();
+		const FString ClassName = IUsdClassesModule::GetClassNameForAnalytics(&Material);
 
 		TArray<FAnalyticsEventAttribute> EventAttributes;
 		EventAttributes.Emplace(TEXT("AssetType"), ClassName);
diff --git a/Engine/Plugins/Importers/USDImporter/Source/USDExporter/Private/SkeletalMeshExporterUSD.cpp b/Engine/Plugins/Importers/USDImporter/Source/USDExporter/Private/SkeletalMeshExporterUSD.cpp
--- a/Engine/Plugins/Importers/USDImporter/Source/USDExporter/Private/SkeletalMeshExporterUSD.cpp
+++ b/Engine/Plugins/Importers/USDImporter/Source/USDExporter/Private/SkeletalMeshExporterUSD.cpp
@@ -47,7 +47,7 @@ namespace UE::SkeletalMeshExporterUSD::Private
 			return;
 		}
 
-		FString ClassName = Asset->GetClass()->GetName();
+		const FString ClassName = IUsdClassesModule::GetClassNameForAnalytics(Asset);
 
 		TArray<FAnalyticsEventAttribute> EventAttributes;
 		EventAttributes.Emplace(TEXT("AssetType"), ClassName);
diff --git a/Engine/Plugins/Importers/USDImporter/Source/USDExporter/Private/StaticMeshExporterUSD.cpp b/Engine/Plugins/Importers/USDImporter/Source/USDExporter/Private/StaticMeshExporterUSD.cpp
--- a/Engine/Plugins/Importers/USDImporter/Source/USDExporter/Private/StaticMeshExporterUSD.cpp
+++ b/Engine/Plugins/Importers/USDImporter/Source/USDExporter/Private/StaticMeshExporterUSD.cpp
@@ -45,7 +45,7 @@ namespace UE::StaticMeshExporterUSD::Private
 			return;
 		}
 
-		FString ClassName = Asset->GetClass()->GetName();
+		const FString ClassName = IUsdClassesModule::GetClassNameForAnalytics(Asset);
 
 		TArray<FAnalyticsEventAttribute> EventAttributes;
 		EventAttributes.Emplace(TEXT("AssetType"), ClassName);
diff --git a/Engine/Plugins/Importers/USDImporter/Source/USDStageImporter/Private/USDStageImporter.cpp b/Engine/Plugins/Importers/USDImporter/Source/USDStageImporter/Private/USDStageImporter.cpp
--- a/Engine/Plugins/Importers/USDImporter/Source/USDStageImporter/Private/USDStageImporter.cpp
+++ b/Engine/Plugins/Importers/USDImporter/Source/USDStageImporter/Private/USDStageImporter.cpp
@@ -1545,7 +1545,7 @@ namespace UE::USDStageImporter::Private
 			FString EventName = Operation;
 			if (Asset)
 			{
-				FString ClassName = Asset->GetClass()->GetName();
+				const FString ClassName = IUsdClassesModule::GetClassNameForAnalytics(Asset);
 
 				// e.g. "Reimport.StaticMesh"
 				EventName = FString::Printf(TEXT("%s.%s"), *EventName, *ClassName);
diff --git a/Engine/Plugins/Runtime/USDCore/Source/USDClasses/Private/USDClassesModule.cpp b/Engine/Plugins/Runtime/USDCore/Source/USDClasses/Private/USDClassesModule.cpp
--- a/Engine/Plugins/Runtime/USDCore/Source/USDClasses/Private/USDClassesModule.cpp
+++ b/Engine/Plugins/Runtime/USDCore/Source/USDClasses/Private/USDClassesModule.cpp
@@ -79,8 +79,10 @@ namespace UE::USDClasses::Private
 	}
 
 	TMap<FString, int32> PackagePathNameToDirtyCounter;
+
+	TMap<FString, TArray<FAnalyticsEventAttribute>> AccumulatedEventAttributes;
 	int AnalyticsAllowedWhenZero = 0;
-}
+}	 // namespace UE::USDClasses::Private
 
 void IUsdClassesModule::UpdatePlugInfoFiles(const FString& PluginDirectory, const FString& TargetDllFolder)
 {
@@ -241,10 +243,15 @@ void IUsdClassesModule::SendAnalytics(
 	const FString& Extension
 )
 {
-	if (FEngineAnalytics::IsAvailable() && UE::USDClasses::Private::AnalyticsAllowedWhenZero == 0)
+	if (!FEngineAnalytics::IsAvailable())
 	{
-		TArray<FAnalyticsEventAttribute> Attributes(InAttributes);
+		return;
+	}
 
+	// Send the event right away
+	if (UE::USDClasses::Private::AnalyticsAllowedWhenZero == 0)
+	{
+		TArray<FAnalyticsEventAttribute> Attributes(InAttributes);
 		Attributes.Emplace(TEXT("Automated"), LexToString(bAutomated));
 		Attributes.Emplace(TEXT("FileExtension"), Extension);
 		Attributes.Emplace(TEXT("NumberOfFrames"), LexToString(FMath::Abs(NumberOfFrames)));
@@ -253,23 +260,102 @@ void IUsdClassesModule::SendAnalytics(
 		Attributes.Emplace(TEXT("EngineVersion"), FEngineVersion::Current().ToString());
 		Attributes.Emplace(TEXT("EngineMode"), FPlatformMisc::GetEngineMode());
 
-		const FString EventText = FString::Printf(TEXT("Engine.Usage.USD.%s"), *EventName);
-		FEngineAnalytics::GetProvider().RecordEvent(EventText, Attributes);
+		const FString EventFullName = FString::Printf(TEXT("Engine.Usage.USD.%s"), *EventName);
+		FEngineAnalytics::GetProvider().RecordEvent(EventFullName, Attributes);
+	}
+	// Something is suppressing analytics, so just accumulate the events
+	else
+	{
+		UE::USDClasses::Private::AccumulatedEventAttributes.FindOrAdd(EventName, InAttributes);
 	}
 }
 
 void IUsdClassesModule::SendAnalytics(TArray<FAnalyticsEventAttribute>&& InAttributes, const FString& EventName)
 {
-	if (FEngineAnalytics::IsAvailable() && UE::USDClasses::Private::AnalyticsAllowedWhenZero == 0)
+	if (!FEngineAnalytics::IsAvailable())
 	{
-		TArray<FAnalyticsEventAttribute> Attributes(InAttributes);
+		return;
+	}
 
+	// Send the event right away
+	if (UE::USDClasses::Private::AnalyticsAllowedWhenZero == 0)
+	{
+		TArray<FAnalyticsEventAttribute> Attributes(InAttributes);
 		Attributes.Emplace(TEXT("Platform"), FPlatformProperties::IniPlatformName());
 		Attributes.Emplace(TEXT("EngineVersion"), FEngineVersion::Current().ToString());
 		Attributes.Emplace(TEXT("EngineMode"), FPlatformMisc::GetEngineMode());
 
-		const FString EventText = FString::Printf(TEXT("Engine.Usage.USD.%s"), *EventName);
-		FEngineAnalytics::GetProvider().RecordEvent(EventText, Attributes);
+		const FString EventFullName = FString::Printf(TEXT("Engine.Usage.USD.%s"), *EventName);
+		FEngineAnalytics::GetProvider().RecordEvent(EventFullName, Attributes);
+	}
+	// Something is suppressing analytics, so just accumulate the events
+	else
+	{
+		UE::USDClasses::Private::AccumulatedEventAttributes.FindOrAdd(EventName, InAttributes);
+	}
+}
+
+FString IUsdClassesModule::GetClassNameForAnalytics(const UObject* Object)
+{
+	if (!Object)
+	{
+		return TEXT("None");
+	}
+	else if (Object->IsA<UStaticMesh>())
+	{
+		return TEXT("StaticMesh");
+	}
+	else if (Object->IsA<USkeletalMesh>())
+	{
+		return TEXT("SkeletalMesh");
+	}
+	else if (Object->IsA<UAnimSequence>())
+	{
+		return TEXT("AnimSequence");
+	}
+	else if (Object->IsA<UMaterialInstance>())
+	{
+		return TEXT("MaterialInstance");
+	}
+	else if (Object->IsA<UMaterial>())
+	{
+		return TEXT("Material");
+	}
+	else if (Object->IsA<UGeometryCache>())
+	{
+		return TEXT("GeometryCache");
+	}
+	else if (Object->IsA<ULevelSequence>())
+	{
+		return TEXT("LevelSequence");
+	}
+	else if (Object->IsA<UTexture>())
+	{
+		return TEXT("Texture");
+	}
+	else if (Object->IsA<UGroomAsset>())
+	{
+		return TEXT("GroomAsset");
+	}
+	else if (Object->IsA<UGroomBindingAsset>())
+	{
+		return TEXT("GroomBindingAsset");
+	}
+	else if (Object->IsA<UGroomCache>())
+	{
+		return TEXT("GroomCache");
+	}
+	else if (Object->IsA<USparseVolumeTexture>())
+	{
+		return TEXT("SparseVolumeTexture");
+	}
+	else if (Object->IsA<USoundBase>())
+	{
+		return TEXT("SoundBase");
+	}
+	else
+	{
+		return TEXT("Other");
 	}
 }
 
@@ -361,6 +447,12 @@ void IUsdClassesModule::AddAssetCountAttributes(const TSet<UObject*>& Assets, TA
 
 void IUsdClassesModule::BlockAnalyticsEvents()
 {
+	// Clear our accumulated events when first going from unblocked to blocked
+	if (UE::USDClasses::Private::AnalyticsAllowedWhenZero == 0)
+	{
+		UE::USDClasses::Private::AccumulatedEventAttributes.Reset();
+	}
+
 	UE::USDClasses::Private::AnalyticsAllowedWhenZero += 1;
 }
 
@@ -369,6 +461,11 @@ void IUsdClassesModule::ResumeAnalyticsEvents()
 	UE::USDClasses::Private::AnalyticsAllowedWhenZero -= 1;
 }
 
+const TMap<FString, TArray<FAnalyticsEventAttribute>>& IUsdClassesModule::GetAccumulatedAnalytics()
+{
+	return UE::USDClasses::Private::AccumulatedEventAttributes;
+}
+
 bool IUsdClassesModule::HashObjectPackage(const UObject* Object, FSHA1& HashToUpdate)
 {
 #if WITH_EDITOR
diff --git a/Engine/Plugins/Runtime/USDCore/Source/USDClasses/Public/USDClassesModule.h b/Engine/Plugins/Runtime/USDCore/Source/USDClasses/Public/USDClassesModule.h
--- a/Engine/Plugins/Runtime/USDCore/Source/USDClasses/Public/USDClassesModule.h
+++ b/Engine/Plugins/Runtime/USDCore/Source/USDClasses/Public/USDClassesModule.h
@@ -37,11 +37,13 @@ public:
 	);
 	USDCLASSES_API static void SendAnalytics(TArray<FAnalyticsEventAttribute>&& InAttributes, const FString& EventName);
 
+	USDCLASSES_API static FString GetClassNameForAnalytics(const UObject* Object);
+
 	USDCLASSES_API static void AddAssetCountAttributes(const TSet<UObject*>& Assets, TArray<FAnalyticsEventAttribute>& InOutAttributes);
 
 	USDCLASSES_API static void BlockAnalyticsEvents();
-
 	USDCLASSES_API static void ResumeAnalyticsEvents();
+	USDCLASSES_API static const TMap<FString, TArray<FAnalyticsEventAttribute>>& GetAccumulatedAnalytics();
 
 	/**
 	 * Updates HashToUpdate with the Object's package's persistent guid, the corresponding file save
