USD: Implement a fallback for collecting mesh data in case it is authored all on a single time sample.

diff --git a/Engine/Plugins/Runtime/USDCore/Source/USDUtilities/Private/USDGeomMeshConversion.cpp b/Engine/Plugins/Runtime/USDCore/Source/USDUtilities/Private/USDGeomMeshConversion.cpp
--- a/Engine/Plugins/Runtime/USDCore/Source/USDUtilities/Private/USDGeomMeshConversion.cpp
+++ b/Engine/Plugins/Runtime/USDCore/Source/USDUtilities/Private/USDGeomMeshConversion.cpp
@@ -1421,6 +1421,37 @@ namespace UE::UsdGeomMeshConversion::Private
 			}
 		}
 
+		auto GetOrEarliest = [&Prim](const auto& AttributeOrPrimvar, auto& OutputArray, pxr::UsdTimeCode TimeCode)
+		{
+			const bool bResult = AttributeOrPrimvar.Get(&OutputArray, TimeCode);
+
+			// Some DCCs seem to output all their mesh data in a single time sample sometimes, as opposed to on the default time code.
+			// When importing, by default we'll exclusively sample the Default timecode, which means we won't find anything in that
+			// case, and so Get() and the mesh parsing as a whole will likely fail.
+			// Here we use a fallback to also check the earliest time code in that case.
+			if (!bResult && TimeCode.IsDefault() && AttributeOrPrimvar.HasAuthoredValue())
+			{
+				if (AttributeOrPrimvar.Get(&OutputArray, pxr::UsdTimeCode::EarliestTime()))
+				{
+					USD_LOG_WARNING(
+						TEXT(
+							"Mesh data attribute or primvar '%s' of prim '%s' exclusively has time samples but was sampled at the default time code, which failed. The importer will fallback to sampling the earliest time code instead."
+						),
+						*UsdToUnreal::ConvertToken(AttributeOrPrimvar.GetName()),
+						*UsdToUnreal::ConvertPath(Prim.GetPrimPath())
+					);
+				}
+			}
+		};
+		auto GetIndicesOrEarliest = [&Prim](const auto& AttributeOrPrimvar, auto& OutputArray, pxr::UsdTimeCode TimeCode)
+		{
+			const bool bResult = AttributeOrPrimvar.GetIndices(&OutputArray, TimeCode);
+			if (!bResult && TimeCode.IsDefault() && AttributeOrPrimvar.HasAuthoredValue())
+			{
+				AttributeOrPrimvar.GetIndices(&OutputArray, pxr::UsdTimeCode::EarliestTime());
+			}
+		};
+
 		// All pointsBased/Gprim attributes we'll retrieve happen to have default varying interpolation
 		const pxr::TfToken DefaultInterpolation = pxr::UsdGeomTokens->varying;
 
@@ -1431,17 +1462,17 @@ namespace UE::UsdGeomMeshConversion::Private
 		// DisplayColors
 		if (pxr::UsdGeomPrimvar DisplayColorsPrimvar = Gprim.GetDisplayColorPrimvar())
 		{
-			DisplayColorsPrimvar.Get(&InOutMeshData.DisplayColors, Options.TimeCode);
+			GetOrEarliest(DisplayColorsPrimvar, InOutMeshData.DisplayColors, Options.TimeCode);
+			GetIndicesOrEarliest(DisplayColorsPrimvar, InOutMeshData.DisplayColorIndices, Options.TimeCode);
 			InOutMeshData.DisplayColorInterpolation = GetAttrInterpolation(DisplayColorsPrimvar, DefaultInterpolation);
-			DisplayColorsPrimvar.GetIndices(&InOutMeshData.DisplayColorIndices, Options.TimeCode);
 		}
 
 		// DisplayOpacities
 		if (pxr::UsdGeomPrimvar DisplayOpacitiesPrimvar = Gprim.GetDisplayOpacityPrimvar())
 		{
-			DisplayOpacitiesPrimvar.Get(&InOutMeshData.DisplayOpacities, Options.TimeCode);
+			GetOrEarliest(DisplayOpacitiesPrimvar, InOutMeshData.DisplayOpacities, Options.TimeCode);
+			GetIndicesOrEarliest(DisplayOpacitiesPrimvar, InOutMeshData.DisplayOpacityIndices, Options.TimeCode);
 			InOutMeshData.DisplayOpacityInterpolation = GetAttrInterpolation(DisplayOpacitiesPrimvar, DefaultInterpolation);
-			DisplayOpacitiesPrimvar.GetIndices(&InOutMeshData.DisplayOpacityIndices, Options.TimeCode);
 		}
 
 		if (pxr::UsdGeomMesh UsdMesh{Prim})
@@ -1449,13 +1480,13 @@ namespace UE::UsdGeomMeshConversion::Private
 			// Faces
 			if (pxr::UsdAttribute FaceVertexCountsAttr = UsdMesh.GetFaceVertexCountsAttr())
 			{
-				FaceVertexCountsAttr.Get(&InOutMeshData.FaceVertexCounts, Options.TimeCode);
+				GetOrEarliest(FaceVertexCountsAttr, InOutMeshData.FaceVertexCounts, Options.TimeCode);
 			}
 
 			// Vertex indices
 			if (pxr::UsdAttribute FaceVertexIndicesAttr = UsdMesh.GetFaceVertexIndicesAttr())
 			{
-				FaceVertexIndicesAttr.Get(&InOutMeshData.FaceVertexIndices, Options.TimeCode);
+				GetOrEarliest(FaceVertexIndicesAttr, InOutMeshData.FaceVertexIndices, Options.TimeCode);
 			}
 
 			// Points
@@ -1463,13 +1494,13 @@ namespace UE::UsdGeomMeshConversion::Private
 			{
 				// Should points always have "vertex" interpolation? Having "varying" forces it to just tessellate instead,
 				// and all OpenSubdiv tutorials use the vertex interpolation type for it
-				PointsPrimvar.Get(&InOutMeshData.Points, Options.TimeCode);
+				GetOrEarliest(PointsPrimvar, InOutMeshData.Points, Options.TimeCode);
+				GetIndicesOrEarliest(PointsPrimvar, InOutMeshData.PointIndices, Options.TimeCode);
 				InOutMeshData.PointInterpolation = GetAttrInterpolation(PointsPrimvar, pxr::UsdGeomTokens->vertex);
-				PointsPrimvar.GetIndices(&InOutMeshData.PointIndices, Options.TimeCode);
 			}
 			else if (pxr::UsdAttribute PointsAttr = UsdMesh.GetPointsAttr())
 			{
-				PointsAttr.Get(&InOutMeshData.Points, Options.TimeCode);
+				GetOrEarliest(PointsAttr, InOutMeshData.Points, Options.TimeCode);
 				InOutMeshData.PointInterpolation = GetAttrInterpolation(PointsAttr, pxr::UsdGeomTokens->vertex);
 			}
 
@@ -1482,20 +1513,20 @@ namespace UE::UsdGeomMeshConversion::Private
 			// Normals
 			if (pxr::UsdGeomPrimvar NormalsPrimvar = pxr::UsdGeomPrimvar(Prim.GetAttribute(UnrealIdentifiers::PrimvarsNormals)))
 			{
-				NormalsPrimvar.Get(&InOutMeshData.Normals, Options.TimeCode);
+				GetOrEarliest(NormalsPrimvar, InOutMeshData.Normals, Options.TimeCode);
+				GetIndicesOrEarliest(NormalsPrimvar, InOutMeshData.NormalIndices, Options.TimeCode);
 				InOutMeshData.NormalInterpolation = GetAttrInterpolation(NormalsPrimvar, DefaultInterpolation);
-				NormalsPrimvar.GetIndices(&InOutMeshData.NormalIndices, Options.TimeCode);
 			}
 			else if (pxr::UsdAttribute NormalsAttr = UsdMesh.GetNormalsAttr())
 			{
-				NormalsAttr.Get(&InOutMeshData.Normals, Options.TimeCode);
+				GetOrEarliest(NormalsAttr, InOutMeshData.Normals, Options.TimeCode);
 				InOutMeshData.NormalInterpolation = GetAttrInterpolation(NormalsAttr, DefaultInterpolation);
 			}
 
 			// Velocities
 			if (pxr::UsdAttribute VelocitiesAttr = UsdMesh.GetVelocitiesAttr())
 			{
-				VelocitiesAttr.Get(&InOutMeshData.Velocities, Options.TimeCode);
+				GetOrEarliest(VelocitiesAttr, InOutMeshData.Velocities, Options.TimeCode);
 				InOutMeshData.VelocityInterpolation = GetAttrInterpolation(VelocitiesAttr, DefaultInterpolation);
 			}
 
@@ -1510,54 +1541,54 @@ namespace UE::UsdGeomMeshConversion::Private
 			{
 				if (pxr::UsdAttribute SubdivSchemeAttr = UsdMesh.GetSubdivisionSchemeAttr())
 				{
-					SubdivSchemeAttr.Get(&InOutMeshData.SubdivScheme, Options.TimeCode);
+					GetOrEarliest(SubdivSchemeAttr, InOutMeshData.SubdivScheme, Options.TimeCode);
 				}
 
 				if (InOutMeshData.SubdivScheme != pxr::UsdGeomTokens->none)
 				{
 					if (pxr::UsdAttribute InterpolateBoundaryAttr = UsdMesh.GetInterpolateBoundaryAttr())
 					{
-						InterpolateBoundaryAttr.Get(&InOutMeshData.InterpolateBoundary, Options.TimeCode);
+						GetOrEarliest(InterpolateBoundaryAttr, InOutMeshData.InterpolateBoundary, Options.TimeCode);
 					}
 
 					if (pxr::UsdAttribute FaceVaryingInterpolationAttr = UsdMesh.GetFaceVaryingLinearInterpolationAttr())
 					{
-						FaceVaryingInterpolationAttr.Get(&InOutMeshData.FaceVaryingInterpolation, Options.TimeCode);
+						GetOrEarliest(FaceVaryingInterpolationAttr, InOutMeshData.FaceVaryingInterpolation, Options.TimeCode);
 					}
 
 					if (pxr::UsdAttribute TriangleSubdivisionAttr = UsdMesh.GetTriangleSubdivisionRuleAttr())
 					{
-						TriangleSubdivisionAttr.Get(&InOutMeshData.TriangleSubdivision, Options.TimeCode);
+						GetOrEarliest(TriangleSubdivisionAttr, InOutMeshData.TriangleSubdivision, Options.TimeCode);
 					}
 
 					if (pxr::UsdAttribute CornerIndicesAttr = UsdMesh.GetCornerIndicesAttr())
 					{
-						CornerIndicesAttr.Get(&InOutMeshData.CornerIndices, Options.TimeCode);
+						GetOrEarliest(CornerIndicesAttr, InOutMeshData.CornerIndices, Options.TimeCode);
 					}
 
 					if (pxr::UsdAttribute CornerSharpnessesAttr = UsdMesh.GetCornerSharpnessesAttr())
 					{
-						CornerSharpnessesAttr.Get(&InOutMeshData.CornerSharpnesses, Options.TimeCode);
+						GetOrEarliest(CornerSharpnessesAttr, InOutMeshData.CornerSharpnesses, Options.TimeCode);
 					}
 
 					if (pxr::UsdAttribute CreaseIndicesAttr = UsdMesh.GetCreaseIndicesAttr())
 					{
-						CreaseIndicesAttr.Get(&InOutMeshData.CreaseIndices, Options.TimeCode);
+						GetOrEarliest(CreaseIndicesAttr, InOutMeshData.CreaseIndices, Options.TimeCode);
 					}
 
 					if (pxr::UsdAttribute CreaseLengthsAttr = UsdMesh.GetCreaseLengthsAttr())
 					{
-						CreaseLengthsAttr.Get(&InOutMeshData.CreaseLengths, Options.TimeCode);
+						GetOrEarliest(CreaseLengthsAttr, InOutMeshData.CreaseLengths, Options.TimeCode);
 					}
 
 					if (pxr::UsdAttribute CreaseSharpnessesAttr = UsdMesh.GetCreaseSharpnessesAttr())
 					{
-						CreaseSharpnessesAttr.Get(&InOutMeshData.CreaseSharpnesses, Options.TimeCode);
+						GetOrEarliest(CreaseSharpnessesAttr, InOutMeshData.CreaseSharpnesses, Options.TimeCode);
 					}
 
 					if (pxr::UsdAttribute HoleIndicesAttr = UsdMesh.GetHoleIndicesAttr())
 					{
-						HoleIndicesAttr.Get(&InOutMeshData.HoleIndices, Options.TimeCode);
+						GetOrEarliest(HoleIndicesAttr, InOutMeshData.HoleIndices, Options.TimeCode);
 					}
 
 					// For some reason this is not part of the USD schema so just pick the first valid token.
@@ -1610,8 +1641,8 @@ namespace UE::UsdGeomMeshConversion::Private
 				// ordering between the cubes
 				if (pxr::UsdGeomPrimvar Primvar = PrimvarPtr.Get())
 				{
-					ensure(Primvar.Get(&UVs, Options.TimeCode));
-					ensure(Primvar.GetIndices(&Indices, Options.TimeCode) || !Primvar.IsIndexed());
+					GetOrEarliest(Primvar, UVs, Options.TimeCode);
+					GetIndicesOrEarliest(Primvar, Indices, Options.TimeCode);
 					Interpolation = GetAttrInterpolation(Primvar, DefaultInterpolation);
 				}
 			}
@@ -1675,7 +1706,9 @@ namespace UE::UsdGeomMeshConversion::Private
 						if (pxr::UsdAttribute SkeletonJointsAttr = BoundSkeleton.GetJointsAttr())
 						{
 							pxr::VtArray<pxr::TfToken> SkeletonJoints;
-							if (SkeletonJointsAttr.Get(&SkeletonJoints))
+							GetOrEarliest(SkeletonJointsAttr, SkeletonJoints, Options.TimeCode);
+
+							if (SkeletonJoints.size() > 0)
 							{
 								// If the skeleton has N bones, this will just contain { 0, 1, 2, ..., N-1 }
 								int NumUsdSkeletonBones = static_cast<int>(SkeletonJoints.size());
