USD MaterialX:
- Fix wrong output connection

diff --git a/Engine/Plugins/Runtime/USDCore/Source/USDUtilities/Private/USDMaterialXShaderGraph.cpp b/Engine/Plugins/Runtime/USDCore/Source/USDUtilities/Private/USDMaterialXShaderGraph.cpp
--- a/Engine/Plugins/Runtime/USDCore/Source/USDUtilities/Private/USDMaterialXShaderGraph.cpp
+++ b/Engine/Plugins/Runtime/USDCore/Source/USDUtilities/Private/USDMaterialXShaderGraph.cpp
@@ -223,16 +223,11 @@ bool FUsdMaterialXShaderGraph::ComputeShaderGraph(mx::NodePtr Node, const pxr::U
 			}
 
 			Input->setConnectedNode(ConnectedNode);
-
-			// add the output of the connected node to the current input
-			TUsdStore<std::vector<pxr::UsdShadeOutput>> UsdOutputs = Source.source.GetOutputs();
-			for (const pxr::UsdShadeOutput& UsdOutput : *UsdOutputs)
+			
+			// no need to add the default output
+			if (const std::string& OutputName = Source.sourceName.GetString(); OutputName != "out")
 			{
-				// no need to add the default output
-				if (std::string OutputName = UsdOutput.GetBaseName().GetString(); OutputName != "out")
-				{
-					Input->setOutputString(OutputName);
-				}
+				Input->setOutputString(OutputName);
 			}
 
 			// recurse if it's a newly created node, otherwise the work has already been done
