USD: Fix reading array out of bounds when parsing a skeletal mesh with multiple Mesh prims that have normals.

diff --git a/Engine/Plugins/Runtime/USDCore/Source/USDUtilities/Private/USDSkeletalDataConversion.cpp b/Engine/Plugins/Runtime/USDCore/Source/USDUtilities/Private/USDSkeletalDataConversion.cpp
--- a/Engine/Plugins/Runtime/USDCore/Source/USDUtilities/Private/USDSkeletalDataConversion.cpp
+++ b/Engine/Plugins/Runtime/USDCore/Source/USDUtilities/Private/USDSkeletalDataConversion.cpp
@@ -1947,9 +1947,8 @@ bool UsdToUnreal::ConvertSkinnedMesh(
 					}
 					else
 					{
-						Triangle.TangentZ[FinalCornerIndex] = Normals[SkelMeshWedge.VertexIndex];
+						Triangle.TangentZ[FinalCornerIndex] = Normals[OriginalVertexIndex];
 					}
-					
 				}
 			}
 		}
