USD: Fix materials being parsed as MaterialX even when they had no MaterialX output.

diff --git a/Engine/Plugins/Interchange/Extensions/OpenUSD/Source/Import/Private/InterchangeUsdTranslator.cpp b/Engine/Plugins/Interchange/Extensions/OpenUSD/Source/Import/Private/InterchangeUsdTranslator.cpp
--- a/Engine/Plugins/Interchange/Extensions/OpenUSD/Source/Import/Private/InterchangeUsdTranslator.cpp
+++ b/Engine/Plugins/Interchange/Extensions/OpenUSD/Source/Import/Private/InterchangeUsdTranslator.cpp
@@ -1129,7 +1129,7 @@ namespace UE::InterchangeUsdTranslator::Private
 
 		// Check for any references of MaterialX
 #if WITH_EDITOR
-		if (RenderContext == UnrealIdentifiers::MaterialXRenderContext)
+		if (RenderContext == UnrealIdentifiers::MaterialXRenderContext && UsdUtils::HasSurfaceOutput(Prim, UnrealIdentifiers::MaterialXRenderContext))
 		{
 			TArray<FString> FilePaths = UsdUtils::GetMaterialXFilePaths(Prim);
 			for (const FString& File : FilePaths)
diff --git a/Engine/Plugins/Runtime/USDCore/Source/USDUtilities/Private/USDShadeConversion.cpp b/Engine/Plugins/Runtime/USDCore/Source/USDUtilities/Private/USDShadeConversion.cpp
--- a/Engine/Plugins/Runtime/USDCore/Source/USDUtilities/Private/USDShadeConversion.cpp
+++ b/Engine/Plugins/Runtime/USDCore/Source/USDUtilities/Private/USDShadeConversion.cpp
@@ -3408,6 +3408,25 @@ bool UsdUtils::RemoveUnrealSurfaceOutput(pxr::UsdPrim& MaterialPrim, const UE::F
 	return true;
 }
 
+bool UsdUtils::HasSurfaceOutput(const pxr::UsdPrim& MaterialPrim, const FName& RenderContext)
+{
+	if (!MaterialPrim)
+	{
+		return false;
+	}
+
+	FScopedUsdAllocs UsdAllocs;
+
+	pxr::UsdShadeMaterial ShadeMaterial{MaterialPrim};
+	if (!ShadeMaterial)
+	{
+		return false;
+	}
+
+	pxr::UsdShadeOutput MaterialOutput = ShadeMaterial.GetSurfaceOutput(UnrealToUsd::ConvertToken(*RenderContext.ToString()).Get());
+	return static_cast<bool>(MaterialOutput);
+}
+
 bool UsdUtils::IsMaterialTranslucent(const pxr::UsdShadeMaterial& UsdShadeMaterial)
 {
 	FScopedUsdAllocs UsdAllocs;
diff --git a/Engine/Plugins/Runtime/USDCore/Source/USDUtilities/Public/USDShadeConversion.h b/Engine/Plugins/Runtime/USDCore/Source/USDUtilities/Public/USDShadeConversion.h
--- a/Engine/Plugins/Runtime/USDCore/Source/USDUtilities/Public/USDShadeConversion.h
+++ b/Engine/Plugins/Runtime/USDCore/Source/USDUtilities/Public/USDShadeConversion.h
@@ -276,6 +276,9 @@ namespace UsdUtils
 	UE_DEPRECATED(5.2, "No longer used as UE material assignments are only visible in the 'unreal' render context anyway")
 	USDUTILITIES_API bool RemoveUnrealSurfaceOutput(pxr::UsdPrim& MaterialPrim, const UE::FSdfLayer& LayerToAuthorIn = UE::FSdfLayer{});
 
+	// Returns whether MaterialPrim is an actual Material, and has a surface output authored for the provided render context
+	USDUTILITIES_API bool HasSurfaceOutput(const pxr::UsdPrim& MaterialPrim, const FName& RenderContext);
+
 	/**
 	 * Returns whether the material needs to be rendered with the Translucent rendering mode.
 	 * This function exists because we need this information *before* we pick the right parent for a material instance and properly convert it.
