USDCore: do not invert leftHanded mesh vertex ordering when subdividing

If we're applying subdivision, pxOsd will take care of converting the
orientation from leftHanded to rightHanded:
[REDACTED_PROFILE_URL]/OpenUSD/blob/v25.02a/pxr/imaging/pxOsd/refinerFactory.cpp#L230

In those cases, we do not want to change the vertex ordering to
convert the handed-ness ourselves.

diff --git a/Engine/Plugins/Runtime/USDCore/Source/USDUtilities/Private/USDGeomMeshConversion.cpp b/Engine/Plugins/Runtime/USDCore/Source/USDUtilities/Private/USDGeomMeshConversion.cpp
--- a/Engine/Plugins/Runtime/USDCore/Source/USDUtilities/Private/USDGeomMeshConversion.cpp
+++ b/Engine/Plugins/Runtime/USDCore/Source/USDUtilities/Private/USDGeomMeshConversion.cpp
@@ -3246,7 +3246,10 @@ namespace UE::UsdGeomMeshConversion::Private
 				const int32 CombinedMaterialIndex = LocalToCombinedMaterialSlotIndices[LocalMaterialIndex];
 
 				// Flip geometry if needed
-				if (InMeshData.Orientation == pxr::UsdGeomTokens->leftHanded)
+				// If we're applying subdivision, pxOsd will take care of
+				// converting the orientation from leftHanded to rightHanded,
+				// so only flip if we are NOT subdividing.
+				if (InMeshData.Orientation == pxr::UsdGeomTokens->leftHanded && InOptions.SubdivisionLevel < 1)
 				{
 					for (int32 i = 0; i < CornerInstanceIDs.Num() / 2; ++i)
 					{
