USD: Simplify handling of skinned meshes to just use bind pose on meshes and skeleton.

Skin morph target normal and position offsets with geomBindTransform to always match skinned mesh coordinate space.

diff --git a/Engine/Plugins/Importers/USDImporter/Source/USDSchemas/Private/USDSkelSkeletonTranslator.cpp b/Engine/Plugins/Importers/USDImporter/Source/USDSchemas/Private/USDSkelSkeletonTranslator.cpp
--- a/Engine/Plugins/Importers/USDImporter/Source/USDSchemas/Private/USDSkelSkeletonTranslator.cpp
+++ b/Engine/Plugins/Importers/USDImporter/Source/USDSchemas/Private/USDSkelSkeletonTranslator.cpp
@@ -598,6 +598,8 @@ namespace UsdSkelSkeletonTranslatorImpl
 
 			if (OutBlendShapes)
 			{
+				pxr::GfMatrix4d GeomBindTransform = SkinningQuery.GetGeomBindTransform(Options.TimeCode);
+
 				pxr::SdfPath MeshPrimPath = LODMeshPrim.GetPrimPath();
 
 				pxr::UsdSkelBindingAPI SkelBindingAPI{Stage->GetPrimAtPath(MeshPrimPath)};
@@ -612,7 +614,8 @@ namespace UsdSkelSkeletonTranslatorImpl
 							NumPointsBeforeThisMesh,
 							InOutUsedMorphTargetNames,
 							*OutBlendShapes,
-							Options
+							Options,
+							&GeomBindTransform
 						);
 					}
 				}
diff --git a/Engine/Plugins/Interchange/Extensions/OpenUSD/Source/Import/Private/InterchangeUsdTranslator.cpp b/Engine/Plugins/Interchange/Extensions/OpenUSD/Source/Import/Private/InterchangeUsdTranslator.cpp
--- a/Engine/Plugins/Interchange/Extensions/OpenUSD/Source/Import/Private/InterchangeUsdTranslator.cpp
+++ b/Engine/Plugins/Interchange/Extensions/OpenUSD/Source/Import/Private/InterchangeUsdTranslator.cpp
@@ -71,6 +71,7 @@
 
 #if USE_USD_SDK
 #include "USDIncludesStart.h"
+#include "pxr/usd/usd/timeCode.h"
 #include "pxr/usd/usdGeom/tokens.h"
 #include "pxr/usd/usdLux/tokens.h"
 #include "pxr/usd/usdShade/tokens.h"
@@ -3226,9 +3227,9 @@ namespace UE::InterchangeUsdTranslator::Private
 			// of bake in our skeleton bind pose directly into our skinned mesh, so we really just want to put the
 			// rest pose on the skeleton asset/ReferenceSkeleton
 			constexpr bool bResetCache = false;
-			BoneNode->SetCustomBindPoseLocalTransform(&NodeContainer, Bone.LocalRestTransform, bResetCache);
-			BoneNode->SetCustomTimeZeroLocalTransform(&NodeContainer, Bone.LocalRestTransform, bResetCache);
-			BoneNode->SetCustomLocalTransform(&NodeContainer, Bone.LocalRestTransform, bResetCache);
+			BoneNode->SetCustomBindPoseLocalTransform(&NodeContainer, Bone.LocalBindTransform, bResetCache);
+			BoneNode->SetCustomTimeZeroLocalTransform(&NodeContainer, Bone.LocalBindTransform, bResetCache);
+			BoneNode->SetCustomLocalTransform(&NodeContainer, Bone.LocalBindTransform, bResetCache);
 
 			BoneToUidAndBoneIndex.Add(ConcatBonePath, {BoneNodeUid, BoneIndex});
 
@@ -3874,8 +3875,39 @@ namespace UE::InterchangeUsdTranslator::Private
 			OutMorphTargetName += TEXT("_") + InbetweenName;
 		}
 
+		// Collect GeomBindTransform if we have one
+		FMatrix GeomBindTransform = FMatrix::Identity;
+		{
+			UE::FUsdSkelSkeletonQuery SkelQuery;
+			{
+				FReadScopeLock ReadLock{Impl.CachedTraversalInfoLock};
+
+				const FString& MeshNodeUid = MeshPrefix + MeshPrim.GetPrimPath().GetString();
+				const FTraversalInfo* MeshInfo = Impl.NodeUidToCachedTraversalInfo.Find(MeshNodeUid);
+				if (MeshInfo)
+				{
+					SkelQuery = MeshInfo->ResolveSkelQuery(Impl.UsdStage);
+				}
+			}
+
+			if (SkelQuery)
+			{
+				if (UE::FUsdSkelSkinningQuery SkinningQuery = UsdUtils::CreateSkinningQuery(MeshPrim, SkelQuery))
+				{
+					GeomBindTransform = SkinningQuery.GetGeomBindTransform(Options.TimeCode.GetValue());
+				}
+			}
+		}
+
 		const float Weight = 1.0f;
-		return UsdUtils::ApplyBlendShape(OutMeshDescription, BlendShape.GetPrim(), Options.AdditionalTransform, Weight, InbetweenName);
+		return UsdUtils::ApplyBlendShape(
+			OutMeshDescription,
+			BlendShape.GetPrim(),
+			GeomBindTransform,
+			Options.AdditionalTransform,
+			Weight,
+			InbetweenName
+		);
 	}
 
 	// Volume animations need a special property reader because in USD it's an animation of the file path attribute within
diff --git a/Engine/Plugins/Runtime/USDCore/Source/USDUtilities/Private/USDGeomMeshConversion.cpp b/Engine/Plugins/Runtime/USDCore/Source/USDUtilities/Private/USDGeomMeshConversion.cpp
--- a/Engine/Plugins/Runtime/USDCore/Source/USDUtilities/Private/USDGeomMeshConversion.cpp
+++ b/Engine/Plugins/Runtime/USDCore/Source/USDUtilities/Private/USDGeomMeshConversion.cpp
@@ -1393,34 +1393,6 @@ namespace UE::UsdGeomMeshConversion::Private
 
 		const bool bProcessSkeletalAttributes = SkeletonQuery && SkeletonQuery->IsValid() && SkinningQuery && SkinningQuery->IsValid();
 
-		pxr::VtArray<pxr::GfMatrix4d> MeshToSkeletonRestPose;
-		if (bProcessSkeletalAttributes)
-		{
-			bool bSuccess = true;
-
-			// Get world-space restTransforms
-			pxr::VtArray<pxr::GfMatrix4d> WorldSpaceRestTransforms;
-			{
-				pxr::VtArray<pxr::GfMatrix4d> JointLocalRestTransforms;
-				const bool bAtRest = true;
-				bSuccess &= SkeletonQuery->ComputeJointLocalTransforms(&JointLocalRestTransforms, Options.TimeCode, bAtRest);
-				bSuccess &= pxr::UsdSkelConcatJointTransforms(SkeletonQuery->GetTopology(), JointLocalRestTransforms, &WorldSpaceRestTransforms);
-			}
-
-			// Get world-space bindTransforms
-			pxr::VtArray<pxr::GfMatrix4d> WorldSpaceBindTransforms;
-			bSuccess &= SkeletonQuery->GetJointWorldBindTransforms(&WorldSpaceBindTransforms);
-
-			if (bSuccess)
-			{
-				MeshToSkeletonRestPose.resize(WorldSpaceRestTransforms.size());
-				for (uint32 Index = 0; Index < WorldSpaceRestTransforms.size(); ++Index)
-				{
-					MeshToSkeletonRestPose[Index] = WorldSpaceBindTransforms[Index].GetInverse() * WorldSpaceRestTransforms[Index];
-				}
-			}
-		}
-
 		auto GetOrEarliest = [&Prim](const auto& AttributeOrPrimvar, auto& OutputArray, pxr::UsdTimeCode TimeCode)
 		{
 			const bool bResult = AttributeOrPrimvar.Get(&OutputArray, TimeCode);
@@ -1452,6 +1424,29 @@ namespace UE::UsdGeomMeshConversion::Private
 			}
 		};
 
+		// Collect GeomBindTransform if we have one
+		bool bApplyGeomBindTransform = false;
+		pxr::GfMatrix4d GeomBindTransform{1.0};
+		pxr::GfMatrix4d InvTransposeGeomBindTransform{1.0};
+		if (bProcessSkeletalAttributes && SkinningQuery)
+		{
+			GeomBindTransform = SkinningQuery->GetGeomBindTransform(Options.TimeCode);
+			if (GeomBindTransform != pxr::GfMatrix4d(1.0))	  // This is USD's identity
+			{
+				if (GeomBindTransform.GetDeterminant() == 0.0)
+				{
+					// Can't invert, just use as-is
+					USD_LOG_WARNING(TEXT("Failed to invert geomBindTransform for prim '%s'"), *UsdToUnreal::ConvertPath(Prim.GetPrimPath()));
+					InvTransposeGeomBindTransform = GeomBindTransform;
+				}
+				else
+				{
+					InvTransposeGeomBindTransform = GeomBindTransform.GetInverse().GetTranspose();
+				}
+				bApplyGeomBindTransform = true;
+			}
+		}
+
 		// All pointsBased/Gprim attributes we'll retrieve happen to have default varying interpolation
 		const pxr::TfToken DefaultInterpolation = pxr::UsdGeomTokens->varying;
 
@@ -1478,7 +1473,7 @@ namespace UE::UsdGeomMeshConversion::Private
 		if (pxr::UsdGeomMesh UsdMesh{Prim})
 		{
 			// Faces
-			if (pxr::UsdAttribute FaceVertexCountsAttr = UsdMesh.GetFaceVertexCountsAttr())
+			if (pxr::UsdAttribute FaceVertexCountsAttr = UsdMesh.GetFaceVertexCountsAttr())  
 			{
 				GetOrEarliest(FaceVertexCountsAttr, InOutMeshData.FaceVertexCounts, Options.TimeCode);
 			}
@@ -1503,11 +1498,12 @@ namespace UE::UsdGeomMeshConversion::Private
 				GetOrEarliest(PointsAttr, InOutMeshData.Points, Options.TimeCode);
 				InOutMeshData.PointInterpolation = GetAttrInterpolation(PointsAttr, pxr::UsdGeomTokens->vertex);
 			}
-
-			// Apply Skinning if applicable.
-			if (bProcessSkeletalAttributes)
+			if (bApplyGeomBindTransform)
 			{
-				SkinningQuery->ComputeSkinnedPoints(MeshToSkeletonRestPose, &InOutMeshData.Points, Options.TimeCode);
+				for (pxr::GfVec3f& Point : InOutMeshData.Points)
+				{
+					Point = static_cast<pxr::GfVec3f>(GeomBindTransform.Transform(Point));
+				}
 			}
 
 			// Normals
@@ -1522,6 +1518,13 @@ namespace UE::UsdGeomMeshConversion::Private
 				GetOrEarliest(NormalsAttr, InOutMeshData.Normals, Options.TimeCode);
 				InOutMeshData.NormalInterpolation = GetAttrInterpolation(NormalsAttr, DefaultInterpolation);
 			}
+			if (bApplyGeomBindTransform)
+			{
+				for (pxr::GfVec3f& Normal : InOutMeshData.Normals)
+				{
+					Normal = static_cast<pxr::GfVec3f>(InvTransposeGeomBindTransform.TransformDir(Normal));
+				}
+			}
 
 			// Velocities
 			if (pxr::UsdAttribute VelocitiesAttr = UsdMesh.GetVelocitiesAttr())
diff --git a/Engine/Plugins/Runtime/USDCore/Source/USDUtilities/Private/USDSkeletalDataConversion.cpp b/Engine/Plugins/Runtime/USDCore/Source/USDUtilities/Private/USDSkeletalDataConversion.cpp
--- a/Engine/Plugins/Runtime/USDCore/Source/USDUtilities/Private/USDSkeletalDataConversion.cpp
+++ b/Engine/Plugins/Runtime/USDCore/Source/USDUtilities/Private/USDSkeletalDataConversion.cpp
@@ -288,6 +288,7 @@ namespace SkelDataConversionImpl
 		const pxr::VtArray<pxr::GfVec3f>& NormalOffsets,
 		const pxr::VtArray<int>& PointIndices,
 		const FUsdStageInfo& StageInfo,
+		const pxr::GfMatrix4d* InGeomBindTransform,
 		uint32 PointIndexOffset,
 		int32 LODIndex,
 		UsdUtils::FUsdBlendShape& OutBlendShape,
@@ -364,20 +365,50 @@ namespace SkelDataConversionImpl
 			}
 		}
 
+		// Setup GeomBindTransform and invtranspose
+		pxr::GfMatrix4d GeomBindTransform = InGeomBindTransform ? *InGeomBindTransform : pxr::GfMatrix4d{1.0f};
+		pxr::GfMatrix4d InvTransposeGeomBindTransform = GeomBindTransform;
+		if (OutBlendShape.bHasAuthoredTangents && GeomBindTransform != pxr::GfMatrix4d(1.0))	// This is USD's identity
+		{
+			if (GeomBindTransform.GetDeterminant() == 0.0)
+			{
+				// Can't invert, just use as-is
+				USD_LOG_WARNING(TEXT("Failed to invert geomBindTransform when parsing blend shape '%s'"), *Name);
+			}
+			else
+			{
+				InvTransposeGeomBindTransform = GeomBindTransform.GetInverse().GetTranspose();
+			}
+		}
+
 		FMatrix TotalMatrix = Options.AdditionalTransform.ToMatrixWithScale();
 		FMatrix TotalMatrixForNormal = TotalMatrix.Inverse().GetTransposed();
 
 		OutBlendShape.Vertices.SetNumUninitialized(NumOffsets);
 		for (uint32 OffsetIndex = 0; OffsetIndex < NumOffsets; ++OffsetIndex)
 		{
-			const FVector UEOffset = UsdToUnreal::ConvertVector(StageInfo, PointOffsets[OffsetIndex]);
-			const FVector UENormal = OutBlendShape.bHasAuthoredTangents ? FVector(TotalMatrixForNormal.TransformVector(UsdToUnreal::ConvertVector(StageInfo, NormalOffsets[OffsetIndex])).GetSafeNormal())
-																		: FVector(0, 0, 0);
-
 			FMorphTargetDelta& ModifiedVertex = OutBlendShape.Vertices[OffsetIndex];
-			ModifiedVertex.PositionDelta = (FVector3f)UEOffset;
-			ModifiedVertex.TangentZDelta = (FVector3f)UENormal;
 			ModifiedVertex.SourceIdx = BaseIndices[OffsetIndex];
+
+			// Position offset
+			// Note: TransformDir here because even the position deltas are still *deltas* (i.e. vector offsets)
+			pxr::GfVec3d USDPoint = GeomBindTransform.TransformDir(PointOffsets[OffsetIndex]);
+			ModifiedVertex.PositionDelta = (FVector3f)UsdToUnreal::ConvertVector(StageInfo, USDPoint);
+
+			// Normal offset
+			if (OutBlendShape.bHasAuthoredTangents)
+			{
+				pxr::GfVec3d USDNormal = InvTransposeGeomBindTransform.TransformDir(NormalOffsets[OffsetIndex]);
+				const FVector UENormal = FVector(
+					TotalMatrixForNormal.TransformVector(UsdToUnreal::ConvertVector(StageInfo, USDNormal)).GetSafeNormal()
+				);
+				ModifiedVertex.TangentZDelta = (FVector3f)UENormal;
+			}
+			else
+			{
+				// Don't leave it uninitialized
+				ModifiedVertex.TangentZDelta = FVector3f::ZeroVector;
+			}
 		}
 
 		return true;
@@ -1139,21 +1170,14 @@ bool UsdToUnreal::ConvertSkeleton(
 			return true;
 		}
 	}
-
-	// Get the rest transforms that will end up on the skeleton data.
-	// Note that we'll want to use the data in restTransforms, so that the skeletal mesh assumes the pose described
-	// by them when non-animated. See the gigantic comment inside ConvertSkinnedMesh.
-	pxr::VtArray<pxr::GfMatrix4d> JointLocalRestTransforms;
-	const bool bAtRest = true;
-	bool bTransformsComputed = InSkeletonQuery.ComputeJointLocalTransforms(&JointLocalRestTransforms, pxr::UsdTimeCode::EarliestTime(), bAtRest);
-	if (NumBones != JointLocalRestTransforms.size())
+	
+	pxr::VtArray<pxr::GfMatrix4d> JointWorldBindTransforms;
+	bool bTransformsComputed = InSkeletonQuery.GetJointWorldBindTransforms(&JointWorldBindTransforms);
+	if (NumBones != JointWorldBindTransforms.size())
 	{
 		return false;
 	}
-
-	pxr::VtArray<pxr::GfMatrix4d> JointWorldBindTransforms;
-	bTransformsComputed &= InSkeletonQuery.GetJointWorldBindTransforms(&JointWorldBindTransforms);
-
+	
 	pxr::VtArray<pxr::GfMatrix4d> JointLocalBindTransforms;
 	bTransformsComputed &= pxr::UsdSkelComputeJointLocalTransforms(SkelTopology, JointWorldBindTransforms, &JointLocalBindTransforms);
 
@@ -1162,7 +1186,7 @@ bool UsdToUnreal::ConvertSkeleton(
 		pxr::UsdStageWeakPtr Stage = InSkeletonQuery.GetSkeleton().GetPrim().GetStage();
 		const FUsdStageInfo StageInfo(Stage);
 
-		for (uint32 Index = 0; Index < JointLocalRestTransforms.size(); ++Index)
+		for (uint32 Index = 0; Index < JointWorldBindTransforms.size(); ++Index)
 		{
 			FUsdSkeletonData::FBone& Bone = OutConvertedData.Bones[Index];
 
@@ -1178,13 +1202,6 @@ bool UsdToUnreal::ConvertSkeleton(
 			// Note that FBX secretly does this as well, because the FBX SDK's Matrix.GetT(), Matrix.GetQ() and Matrix.GetS()
 			// (used within UnFbx::FFbxImporter::ImportBones) seem to behave the same way and flip all axes when a reflection is
 			// detected.
-			{
-				const pxr::GfMatrix4d& UsdRestTransform = JointLocalRestTransforms[Index];
-				FMatrix Matrix = UsdToUnreal::ConvertMatrix(UsdRestTransform);
-				FTransform RestTransform = UsdUtils::DecomposeWithUniformReflection(Matrix);
-				RestTransform = UsdUtils::ConvertTransformToUESpace(StageInfo, RestTransform);
-				Bone.LocalRestTransform = RestTransform;
-			}
 			{
 				const pxr::GfMatrix4d& UsdBindTransform = JointLocalBindTransforms[Index];
 				FMatrix Matrix = UsdToUnreal::ConvertMatrix(UsdBindTransform);
@@ -1220,7 +1237,6 @@ bool UsdToUnreal::ConvertSkeleton(
 		FUsdSkeletonData::FBone TrueRoot;
 		TrueRoot.Name = UsdUnreal::ObjectUtils::GetUniqueName(TEXT("Root"), BoneNames);
 		TrueRoot.ParentIndex = INDEX_NONE;
-		TrueRoot.LocalRestTransform = FTransform::Identity;
 		OutConvertedData.Bones.Insert(TrueRoot, 0);
 	}
 
@@ -1269,7 +1285,7 @@ bool UsdToUnreal::ConvertSkeleton(const pxr::UsdSkelSkeletonQuery& SkeletonQuery
 		JointMatrix.XSize = 100.f;
 		JointMatrix.YSize = 100.f;
 		JointMatrix.ZSize = 100.f;
-		JointMatrix.Transform = FTransform3f(InBone.LocalRestTransform);
+		JointMatrix.Transform = FTransform3f(InBone.LocalBindTransform);
 	}
 
 	return true;
@@ -1326,91 +1342,45 @@ bool UsdToUnreal::ConvertSkinnedMesh(
 
 	const FUsdStageInfo StageInfo(SkinningPrim.GetStage());
 
-	// Retrieve the mesh points (vertices) from USD and append it to the SkeletalMeshImportData Points
-	uint32 NumPoints = 0;
-	uint32 NumExistingPoints = SkelMeshImportData.Points.Num();
-
-	// In USD, bindTransforms is used to place the skeleton in world space, and geomBindTransforms is used to
-	// place the mesh in world space. This can be understood as "bind time", and the relationship between the joints
-	// and triangles there will then be maintained throughout the skinning/animation process. In UE the analogue of
-	// this "bind time" happens when we just fill in FReferenceSkeleton and USkeletalMesh with data, and the
-	// relationship between triangles and joints there is maintained throughout animation.
-	//
-	// When we convert our skinned mesh and fill in our USkeletalMesh, then, we should have two goals in mind:
-	// - Our USkeletalMesh's final triangle positions must match the final FReferenceSkeleton pose in the same way
-	//   that the USD skeleton matches the Mesh prim at USD "bind time" (so we must use bindTransforms and
-	//   geomBindTransform in some way);
-	// - We want to be able to use the restTransforms pose on our FReferenceSkeleton, so that the rest pose in UE is
-	//   the same rest pose used in USD (and not just the bind pose instead).
-	//
-	// To accomplish this, we'll have our FReferenceSkeleton assume the restTransforms of the USD skeleton,
-	// and then skin the mesh data to deform it according to restTransforms too, so that they match.
-	//
-	// This means that if we're placing restTransforms on our FReferenceSkeleton, we must put our mesh data skinned
-	// by restTransforms on our USkeletalMesh. The problem is that the USD mesh raw points is specified in the "bind
-	// time" pose already (except for an additional, optional transform called geomBindTransform that we can apply).
-	// In order to move the mesh data from its source space to restTransforms pose, we must then:
-	// - Apply geomBindTransform: It is now at "bind time", directly in world space (note: don't confuse this world
-	//   space with the final concatenated prim transform after we account for parent Xforms and so on. We're only
-	//   concerned with "bind time" here, and a different Xform/skinning process will get our final mesh at world
-	//   space. i.e. Here we don't care if our parent SkelRoot has e.g. a +30 translation on the X);
-	// - Invert the skinning done by bindTransforms: This reverses the effect of the bindTransforms pose, but note
-	//   that it most importantly *deforms the mesh according to the skeleton's bindTransforms*. This is essential,
-	//   and this step is what will ensure that the triangles-joints relationship specified by bindTransforms +
-	//   geomBindTransforms is maintained. Practically, this means that the mesh will act as if all joints are on
-	//   top of each other at the origin, and the mesh will become a packed mess of triangles;
-	// - Skin the mesh by restTransforms: This puts the mesh nicely in the final rest pose we want.
-	//
-	// We'll do this by creating a set of skinning transforms (MeshToSkeletonRestPose) that concatenates the two latter
-	// steps, and by calling ComputeSkinnedPoints, which applies geomBindTransform first and then skins the mesh with
-	// those transforms.
-	//
-	// Reference: https://graphics.pixar.com/usd/release/api/_usd_skel__schema_overview.html
-	VtArray<GfMatrix4d> MeshToSkeletonRestPose;
-	{
-		bool bSuccess = true;
-
-		// Get world-space restTransforms
-		VtArray<GfMatrix4d> WorldSpaceRestTransforms;
-		{
-			VtArray<GfMatrix4d> JointLocalRestTransforms;
-			const bool bAtRest = true;
-			bSuccess &= SkeletonQuery.ComputeJointLocalTransforms(&JointLocalRestTransforms, UsdTimeCode::EarliestTime(), bAtRest);
-			bSuccess &= UsdSkelConcatJointTransforms(SkeletonQuery.GetTopology(), JointLocalRestTransforms, &WorldSpaceRestTransforms);
-		}
-
-		// Get world-space bindTransforms
-		VtArray<GfMatrix4d> WorldSpaceBindTransforms;
-		bSuccess &= SkeletonQuery.GetJointWorldBindTransforms(&WorldSpaceBindTransforms);
-
-		if (!bSuccess)
-		{
-			return false;
-		}
-
-		MeshToSkeletonRestPose.resize(WorldSpaceRestTransforms.size());
-		for (uint32 Index = 0; Index < WorldSpaceRestTransforms.size(); ++Index)
+	// Collect GeomBindTransform if we have one
+	pxr::GfMatrix4d GeomBindTransform{1.0};
+	pxr::GfMatrix4d InvTransposeGeomBindTransform{1.0};
+	{
+		GeomBindTransform = SkinningQuery.GetGeomBindTransform(CommonOptions.TimeCode);
+		if (GeomBindTransform != pxr::GfMatrix4d(1.0))	  // This is USD's identity
 		{
-			MeshToSkeletonRestPose[Index] = WorldSpaceBindTransforms[Index].GetInverse() * WorldSpaceRestTransforms[Index];
+			if (GeomBindTransform.GetDeterminant() == 0.0)
+			{
+				// Can't invert, just use as-is
+				USD_LOG_WARNING(TEXT("Failed to invert geomBindTransform for prim '%s'"), *UsdToUnreal::ConvertPath(SkinningPrim.GetPrimPath()));
+				InvTransposeGeomBindTransform = GeomBindTransform;
+			}
+			else
+			{
+				InvTransposeGeomBindTransform = GeomBindTransform.GetInverse().GetTranspose();
+			}
 		}
 	}
 
+	// Retrieve the mesh points (vertices) from USD and append it to the SkeletalMeshImportData Points
+	uint32 NumPoints = 0;
+	uint32 NumExistingPoints = SkelMeshImportData.Points.Num();
+
 	UsdAttribute PointsAttr = UsdMesh.GetPointsAttr();
 	if (PointsAttr)
 	{
 		VtArray<GfVec3f> UsdPoints;
 		PointsAttr.Get(&UsdPoints, UsdTimeCode::Default());
 
-		SkinningQuery.ComputeSkinnedPoints(MeshToSkeletonRestPose, &UsdPoints, UsdTimeCode::EarliestTime());
-
 		NumPoints = UsdPoints.size();
 		SkelMeshImportData.Points.AddUninitialized(NumPoints);
 
 		for (uint32 PointIndex = 0; PointIndex < NumPoints; ++PointIndex)
 		{
-			const GfVec3f& Point = UsdPoints[PointIndex];
-			SkelMeshImportData.Points[PointIndex + NumExistingPoints] = (FVector3f
-			)CommonOptions.AdditionalTransform.TransformPosition(UsdToUnreal::ConvertVector(StageInfo, Point));
+			const GfVec3d Point = GeomBindTransform.Transform(UsdPoints[PointIndex]);
+			SkelMeshImportData.Points[PointIndex + NumExistingPoints] = (FVector3f)CommonOptions.AdditionalTransform.TransformPosition(
+				UsdToUnreal::ConvertVector(StageInfo, Point)
+			);
 		}
 	}
 
@@ -1451,18 +1421,17 @@ bool UsdToUnreal::ConvertSkinnedMesh(
 
 		if (NormalsAttribute.Get(&USDNormals, UsdTimeCode::Default()) && USDNormals.size() > 0)
 		{
-			// Like for points, we need to ensure these normals are in the same coordinate space of the skeleton
-			SkinningQuery.ComputeSkinnedNormals(MeshToSkeletonRestPose, &USDNormals, UsdTimeCode::EarliestTime());
-
 			if (USDNormals.size() == NumPoints || USDNormals.size() == NumVertexInstances)
 			{
 				bUsingVertexInstanceNormals = USDNormals.size() == NumVertexInstances;
 				Normals.Reserve(bUsingVertexInstanceNormals ? NumVertexInstances : NumPoints);
 				SkelMeshImportData.bHasNormals = true;
-				for (GfVec3f& Normal : USDNormals)
+
+				for (GfVec3f& USDNormal : USDNormals)
 				{
-					FVector TransformedNormal = TotalMatrixForNormal.TransformVector(UsdToUnreal::ConvertVector(StageInfo, Normal)).GetSafeNormal();
-					Normals.Add((FVector3f)TransformedNormal);
+					USDNormal = static_cast<pxr::GfVec3f>(InvTransposeGeomBindTransform.TransformDir(USDNormal));
+					FVector UENormal = TotalMatrixForNormal.TransformVector(UsdToUnreal::ConvertVector(StageInfo, USDNormal)).GetSafeNormal();
+					Normals.Add((FVector3f)UENormal);
 				}
 			}
 		}
@@ -2570,7 +2539,8 @@ bool UsdToUnreal::ConvertBlendShape(
 	uint32 PointIndexOffset,
 	TSet<FString>& UsedMorphTargetNames,
 	UsdUtils::FBlendShapeMap& OutBlendShapes,
-	const UsdToUnreal::FUsdMeshConversionOptions& Options
+	const UsdToUnreal::FUsdMeshConversionOptions& Options,
+	const pxr::GfMatrix4d* GeomBindTransform
 )
 {
 	TRACE_CPUPROFILER_EVENT_SCOPE(UsdToUnreal::ConvertBlendShape);
@@ -2610,8 +2580,18 @@ bool UsdToUnreal::ConvertBlendShape(
 	}
 
 	UsdUtils::FUsdBlendShape PrimaryBlendShape;
-	if (!SkelDataConversionImpl::
-			CreateUsdBlendShape(PrimaryName, Offsets, Normals, PointIndices, StageInfo, PointIndexOffset, LODIndex, PrimaryBlendShape, Options))
+	if (!SkelDataConversionImpl::CreateUsdBlendShape(
+			PrimaryName,
+			Offsets,
+			Normals,
+			PointIndices,
+			StageInfo,
+			GeomBindTransform,
+			PointIndexOffset,
+			LODIndex,
+			PrimaryBlendShape,
+			Options
+		))
 	{
 		return false;
 	}
@@ -2658,6 +2638,7 @@ bool UsdToUnreal::ConvertBlendShape(
 				InbetweenNormalOffsets,
 				PointIndices,
 				StageInfo,
+				GeomBindTransform,
 				PointIndexOffset,
 				LODIndex,
 				InbetweenShape,
@@ -3310,6 +3291,27 @@ bool UsdUtils::ApplyBlendShape(
 	float Weight,
 	const FString InInbetweenName
 )
+{
+	const static FMatrix& GeomBindTransform = FMatrix::Identity;
+
+	return ApplyBlendShape(
+		InOutMeshDescription,
+		InBlendShapePrim,
+		GeomBindTransform,
+		AdditionalTransform,
+		Weight,
+		InInbetweenName
+	);
+}
+
+bool UsdUtils::ApplyBlendShape(
+	FMeshDescription& InOutMeshDescription,
+	const pxr::UsdPrim& InBlendShapePrim,
+	const FMatrix& GeomBindTransform,
+	const FTransform& AdditionalTransform,
+	float Weight,
+	const FString InInbetweenName
+)
 {
 	TRACE_CPUPROFILER_EVENT_SCOPE(UsdToUnreal::ConvertBlendShape);
 
@@ -3367,6 +3369,36 @@ bool UsdUtils::ApplyBlendShape(
 		}
 	}
 
+	// Apply GeomBindTransform if we have one
+	if (GeomBindTransform != FMatrix::Identity)
+	{
+		pxr::GfMatrix4d UsdGeomBindTransform = UnrealToUsd::ConvertMatrix(GeomBindTransform);
+		pxr::GfMatrix4d InvTransposeGeomBindTransform;
+		if (UsdGeomBindTransform.GetDeterminant() == 0.0)
+		{
+			// Can't invert, just use as-is
+			USD_LOG_WARNING(
+				TEXT("Failed to invert geomBindTransform for blend shape prim '%s'"),
+				*UsdToUnreal::ConvertPath(InBlendShapePrim.GetPrimPath())
+			);
+			InvTransposeGeomBindTransform = UsdGeomBindTransform;
+		}
+		else
+		{
+			InvTransposeGeomBindTransform = UsdGeomBindTransform.GetInverse().GetTranspose();
+		}
+
+		for (pxr::GfVec3f& Position : PositionOffsets)
+		{
+			// Note: TransformDir here because even the position deltas are still *deltas* (i.e. vector offsets)
+			Position = static_cast<pxr::GfVec3f>(UsdGeomBindTransform.TransformDir(Position));
+		}
+		for (pxr::GfVec3f& Normal : NormalOffsets)
+		{
+			Normal = static_cast<pxr::GfVec3f>(InvTransposeGeomBindTransform.TransformDir(Normal));
+		}
+	}
+
 	FStaticMeshAttributes Attributes{InOutMeshDescription};
 	TVertexAttributesRef<FVector3f> MeshPositions = Attributes.GetVertexPositions();
 	TVertexInstanceAttributesRef<FVector3f> MeshInstanceNormals = Attributes.GetVertexInstanceNormals();
@@ -3374,6 +3406,7 @@ bool UsdUtils::ApplyBlendShape(
 	bool bApplied = false;
 
 	// We have one value for each vertex of the mesh description
+	// TODO: Deduplicate this code, since only the indexing differs?
 	if (PointIndices.size() == 0)
 	{
 		// Position offsets
diff --git a/Engine/Plugins/Runtime/USDCore/Source/USDUtilities/Public/USDSkeletalDataConversion.h b/Engine/Plugins/Runtime/USDCore/Source/USDUtilities/Public/USDSkeletalDataConversion.h
--- a/Engine/Plugins/Runtime/USDCore/Source/USDUtilities/Public/USDSkeletalDataConversion.h
+++ b/Engine/Plugins/Runtime/USDCore/Source/USDUtilities/Public/USDSkeletalDataConversion.h
@@ -181,6 +181,8 @@ namespace UsdUtils
 	 * Note that we expect InOutMeshDescription and InBlendShapePrim to at least have a matching number of vertices/points.
 	 * @param InOutMeshDescription - The mesh data to deform with the blend shape
 	 * @param InBlendShapePrim - The UsdSkelBlendShape prim to deform the mesh with
+	 * @param GeomBindTransform - USD-space transform applied to the point/normal deltas before any other transformation
+	 * @param AdditionalTransform - Unreal-space transform applied to the normal deltas after all other transformations
 	 * @param Weight - Scalar that multiplies the deltas before they're applied to the mesh deformation:
 	 *                 0.0f means no effect, 1.0f adds the deltas at full strength, etc.
 	 * @param InInbetweenName - The name of the inbetween shape of InBlendShapePrim to apply instead, if any.
@@ -194,6 +196,14 @@ namespace UsdUtils
 		float Weight = 1.0f,		
 		const FString InInbetweenName = {}
 	);
+	USDUTILITIES_API bool ApplyBlendShape(
+		FMeshDescription& InOutMeshDescription,
+		const pxr::UsdPrim& InBlendShapePrim,
+		const FMatrix& GeomBindTransform,
+		const FTransform& AdditionalTransform,
+		float Weight = 1.0f,		
+		const FString InInbetweenName = {}
+	);
 #endif	  // USE_USD_SDK
 }	 // namespace UsdUtils
 
@@ -204,6 +214,8 @@ namespace UsdToUnreal
 		struct FBone
 		{
 			FString Name;
+
+			UE_DEPRECATED(5.6, "Skeleton transforms are exclusively derived from the bind pose for now.")
 			FTransform LocalRestTransform = FTransform::Identity;
 			FTransform LocalBindTransform = FTransform::Identity;
 
@@ -250,6 +262,8 @@ namespace UsdToUnreal
 	 * @param UsedMorphTargetNames - Names that the newly created FUsdBlendShapes cannot have (this function will also add the names of the newly
 	 * created FUsdBlendShapes to it)
 	 * @param OutBlendShapes - Where the newly created blend shapes will be added to
+	 * @param Options - Generic mesh conversion options to use when converting the blend shape
+	 * @param GeomBindTransform - USD-space transform to be applied to the point and normal offsets before any other coordinate space conversion
 	 * @return Whether the conversion was successful or not.
 	 */
 	USDUTILITIES_API bool ConvertBlendShape(
@@ -267,7 +281,8 @@ namespace UsdToUnreal
 		uint32 PointIndexOffset,
 		TSet<FString>& UsedMorphTargetNames,
 		UsdUtils::FBlendShapeMap& OutBlendShapes,
-		const FUsdMeshConversionOptions& Options
+		const FUsdMeshConversionOptions& Options,
+		const pxr::GfMatrix4d* GeomBindTransform = nullptr
 	);
 
 	/**
diff --git a/Engine/Plugins/Runtime/USDCore/Source/UnrealUSDWrapper/Private/UsdWrappers/UsdSkelSkinningQuery.cpp b/Engine/Plugins/Runtime/USDCore/Source/UnrealUSDWrapper/Private/UsdWrappers/UsdSkelSkinningQuery.cpp
--- a/Engine/Plugins/Runtime/USDCore/Source/UnrealUSDWrapper/Private/UsdWrappers/UsdSkelSkinningQuery.cpp
+++ b/Engine/Plugins/Runtime/USDCore/Source/UnrealUSDWrapper/Private/UsdWrappers/UsdSkelSkinningQuery.cpp
@@ -178,4 +178,22 @@ namespace UE
 		return bResult;
 	}
 
+	FMatrix FUsdSkelSkinningQuery::GetGeomBindTransform(double UsdTimeCode) const
+	{
+		FMatrix Result = FMatrix::Identity;
+
+#if USE_USD_SDK
+		pxr::GfMatrix4d USDMatrix = Impl->PxrUsdSkelSkinningQuery.Get().GetGeomBindTransform(UsdTimeCode);
+
+		Result = FMatrix(
+			FPlane(USDMatrix[0][0], USDMatrix[0][1], USDMatrix[0][2], USDMatrix[0][3]),
+			FPlane(USDMatrix[1][0], USDMatrix[1][1], USDMatrix[1][2], USDMatrix[1][3]),
+			FPlane(USDMatrix[2][0], USDMatrix[2][1], USDMatrix[2][2], USDMatrix[2][3]),
+			FPlane(USDMatrix[3][0], USDMatrix[3][1], USDMatrix[3][2], USDMatrix[3][3])
+		);
+#endif
+
+		return Result;
+	}
+
 }	 // namespace UE
diff --git a/Engine/Plugins/Runtime/USDCore/Source/UnrealUSDWrapper/Public/UsdWrappers/UsdSkelSkinningQuery.h b/Engine/Plugins/Runtime/USDCore/Source/UnrealUSDWrapper/Public/UsdWrappers/UsdSkelSkinningQuery.h
--- a/Engine/Plugins/Runtime/USDCore/Source/UnrealUSDWrapper/Public/UsdWrappers/UsdSkelSkinningQuery.h
+++ b/Engine/Plugins/Runtime/USDCore/Source/UnrealUSDWrapper/Public/UsdWrappers/UsdSkelSkinningQuery.h
@@ -64,6 +64,9 @@ namespace UE
 
 		bool GetBlendShapeOrder(TArray<FString>& BlendShapes) const;
 
+		// Note: This matrix is in USD space, retrieved as-is 
+		FMatrix GetGeomBindTransform(double UsdTimeCode) const;
+
 	private:
 		TUniquePtr<Internal::FUsdSkelSkinningQueryImpl> Impl;
 	};
