USD: Remove deprecation on FBone LocalRestTransform member as we may end up using it later to support UPoseAsset generation anyway.

CIS Issue 892670

diff --git a/Engine/Plugins/Runtime/USDCore/Source/USDUtilities/Public/USDSkeletalDataConversion.h b/Engine/Plugins/Runtime/USDCore/Source/USDUtilities/Public/USDSkeletalDataConversion.h
--- a/Engine/Plugins/Runtime/USDCore/Source/USDUtilities/Public/USDSkeletalDataConversion.h
+++ b/Engine/Plugins/Runtime/USDCore/Source/USDUtilities/Public/USDSkeletalDataConversion.h
@@ -215,7 +215,6 @@ namespace UsdToUnreal
 		{
 			FString Name;
 
-			UE_DEPRECATED(5.6, "Skeleton transforms are exclusively derived from the bind pose for now.")
 			FTransform LocalRestTransform = FTransform::Identity;
 			FTransform LocalBindTransform = FTransform::Identity;
 
