USDCore: avoid creating clipboard UsdStage when not needed

diff --git a/Engine/Plugins/Runtime/USDCore/Source/USDUtilities/Private/USDConversionUtils.cpp b/Engine/Plugins/Runtime/USDCore/Source/USDUtilities/Private/USDConversionUtils.cpp
--- a/Engine/Plugins/Runtime/USDCore/Source/USDUtilities/Private/USDConversionUtils.cpp
+++ b/Engine/Plugins/Runtime/USDCore/Source/USDUtilities/Private/USDConversionUtils.cpp
@@ -3503,7 +3503,7 @@ TArray<UE::FSdfPath> UsdUtils::PastePrims(const UE::FUsdPrim& ParentPrim)
 		return Result;
 	}
 
-	pxr::UsdStageRefPtr ClipboardStage = pxr::UsdStageRefPtr{UnrealUSDWrapper::GetClipboardStage()};
+	pxr::UsdStageRefPtr ClipboardStage = pxr::UsdStageRefPtr{UnrealUSDWrapper::GetClipboardStage(/* bCreateIfNeeded = */ false)};
 	if (!ClipboardStage)
 	{
 		return Result;
@@ -3579,7 +3579,7 @@ TArray<UE::FSdfPath> UsdUtils::PastePrims(const UE::FUsdPrim& ParentPrim)
 bool UsdUtils::CanPastePrims()
 {
 #if USE_USD_SDK
-	pxr::UsdStageRefPtr ClipboardStage = pxr::UsdStageRefPtr{UnrealUSDWrapper::GetClipboardStage()};
+	pxr::UsdStageRefPtr ClipboardStage = pxr::UsdStageRefPtr{UnrealUSDWrapper::GetClipboardStage(/* bCreateIfNeeded = */ false)};
 	if (!ClipboardStage)
 	{
 		return false;
@@ -3600,7 +3600,7 @@ bool UsdUtils::CanPastePrims()
 void UsdUtils::ClearPrimClipboard()
 {
 #if USE_USD_SDK
-	pxr::UsdStageRefPtr ClipboardStage = pxr::UsdStageRefPtr{UnrealUSDWrapper::GetClipboardStage()};
+	pxr::UsdStageRefPtr ClipboardStage = pxr::UsdStageRefPtr{UnrealUSDWrapper::GetClipboardStage(/* bCreateIfNeeded = */ false)};
 	if (!ClipboardStage)
 	{
 		return;
diff --git a/Engine/Plugins/Runtime/USDCore/Source/UnrealUSDWrapper/Private/UnrealUSDWrapper.cpp b/Engine/Plugins/Runtime/USDCore/Source/UnrealUSDWrapper/Private/UnrealUSDWrapper.cpp
--- a/Engine/Plugins/Runtime/USDCore/Source/UnrealUSDWrapper/Private/UnrealUSDWrapper.cpp
+++ b/Engine/Plugins/Runtime/USDCore/Source/UnrealUSDWrapper/Private/UnrealUSDWrapper.cpp
@@ -985,14 +985,14 @@ UE::FUsdStage UnrealUSDWrapper::NewStage()
 #endif	  // #if USE_USD_SDK
 }
 
-UE::FUsdStage UnrealUSDWrapper::GetClipboardStage()
+UE::FUsdStage UnrealUSDWrapper::GetClipboardStage(bool bCreateIfNeeded)
 {
 	static UE::FUsdStage ClipboardStage;
 
 #if USE_USD_SDK
 	FScopedUsdAllocs UsdAllocs;
 
-	if (!ClipboardStage)
+	if (!ClipboardStage && bCreateIfNeeded)
 	{
 		// Use this specific name so that we can check for it on our automated tests
 		pxr::UsdStageCacheContext StageCacheContext{pxr::UsdUtilsStageCache::Get()};
diff --git a/Engine/Plugins/Runtime/USDCore/Source/UnrealUSDWrapper/Public/UnrealUSDWrapper.h b/Engine/Plugins/Runtime/USDCore/Source/UnrealUSDWrapper/Public/UnrealUSDWrapper.h
--- a/Engine/Plugins/Runtime/USDCore/Source/UnrealUSDWrapper/Public/UnrealUSDWrapper.h
+++ b/Engine/Plugins/Runtime/USDCore/Source/UnrealUSDWrapper/Public/UnrealUSDWrapper.h
@@ -350,9 +350,11 @@ public:
 	/**
 	 * Get the singleton, persistent stage used as a clipboard for prim cut/copy/paste operations.
 	 * WARNING: This stage may remain open indefinitely! If you use this directly, be aware that there may be unintended
-	 * consequences (e.g. a sublayer added to this stage may never fully close)
+	 * consequences (e.g. a sublayer added to this stage may never fully close). When invoking a paste or clear of the
+	 * clipboard stage when it is unknown whether a previous cut or copy was done, setting the bCreateIfNeeded parameter
+	 * to false will avoid needlessly creating the stage.
 	 */
-	UNREALUSDWRAPPER_API static UE::FUsdStage GetClipboardStage();
+	UNREALUSDWRAPPER_API static UE::FUsdStage GetClipboardStage(bool bCreateIfNeeded = true);
 
 	/** Returns all the stages that are currently opened in the USD utils stage cache, shared between C++ and Python */
 	UNREALUSDWRAPPER_API static TArray<UE::FUsdStage> GetAllStagesFromCache();
