USD:
- Fix a crash in the MaterialX shader graph

diff --git a/Engine/Plugins/Runtime/USDCore/Source/USDUtilities/Private/USDMaterialXShaderGraph.cpp b/Engine/Plugins/Runtime/USDCore/Source/USDUtilities/Private/USDMaterialXShaderGraph.cpp
--- a/Engine/Plugins/Runtime/USDCore/Source/USDUtilities/Private/USDMaterialXShaderGraph.cpp
+++ b/Engine/Plugins/Runtime/USDCore/Source/USDUtilities/Private/USDMaterialXShaderGraph.cpp
@@ -252,9 +252,9 @@ bool FUsdMaterialXShaderGraph::IsGeomColor(const pxr::UsdShadeShader& GeomPropVa
 {
 	if (pxr::UsdShadeInput GeomPropInput = GeomPropValueShader.GetInput(pxr::TfToken{"geomprop"}))
 	{
-		if (std::string Value; GeomPropInput.Get<std::string>(&Value))
+		if (TUsdStore<std::string> Value; GeomPropInput.Get<std::string>(&Value.Get()))
 		{
-			if (Value == "displayColor" || Value == "displayOpacity")
+			if (*Value == "displayColor" || *Value == "displayOpacity")
 			{
 				return true;
 			}
