fix UnrealUSDWrapper pulling in Win64 python libraries on WinArm64

#rnx

diff --git a/Engine/Plugins/Runtime/USDCore/Source/UnrealUSDWrapper/UnrealUSDWrapper.Build.cs b/Engine/Plugins/Runtime/USDCore/Source/UnrealUSDWrapper/UnrealUSDWrapper.Build.cs
--- a/Engine/Plugins/Runtime/USDCore/Source/UnrealUSDWrapper/UnrealUSDWrapper.Build.cs
+++ b/Engine/Plugins/Runtime/USDCore/Source/UnrealUSDWrapper/UnrealUSDWrapper.Build.cs
@@ -33,8 +33,9 @@ namespace UnrealBuildTool.Rules
 				PublicDefinitions.Add("ENABLE_USD_DEBUG_PATH=0");
 
 				var EngineDir = Path.GetFullPath(Target.RelativeEnginePath);
-				var PythonSourceTPSDir = Path.Combine(EngineDir, "Source", "ThirdParty", "Python3", Target.Platform.ToString());
-				var PythonBinaryTPSDir = Path.Combine(EngineDir, "Binaries", "ThirdParty", "Python3", Target.Platform.ToString());
+				var PlatformDir = (Target.Architecture == UnrealArch.Arm64 && Target.Platform == UnrealTargetPlatform.Win64) ? "WinArm64" : Target.Platform.ToString();
+				var PythonSourceTPSDir = Path.Combine(EngineDir, "Source", "ThirdParty", "Python3", PlatformDir);
+				var PythonBinaryTPSDir = Path.Combine(EngineDir, "Binaries", "ThirdParty", "Python3", PlatformDir);
 				string USDLibsDir = Path.Combine(ModuleDirectory, "..", "ThirdParty", "USD", "lib");
 
 				if (Target.Platform == UnrealTargetPlatform.Win64)
@@ -42,9 +43,10 @@ namespace UnrealBuildTool.Rules
 					PublicDefinitions.Add("USD_USES_SYSTEM_MALLOC=1");
 
 					// Python3
+					string ArchDir = (Target.Architecture == UnrealArch.Arm64) ? "arm64" : "";
 					PublicSystemIncludePaths.Add(Path.Combine(PythonSourceTPSDir, "include"));
 					PublicSystemLibraryPaths.Add(Path.Combine(PythonSourceTPSDir, "libs"));
-					RuntimeDependencies.Add(Path.Combine("$(TargetOutputDir)", "python311.dll"), Path.Combine(PythonBinaryTPSDir, "python311.dll"));
+					RuntimeDependencies.Add(Path.Combine("$(TargetOutputDir)", ArchDir, "python311.dll"), Path.Combine(PythonBinaryTPSDir, "python311.dll"));
 
 					// USD
 					PublicSystemIncludePaths.Add(Path.Combine(ModuleDirectory, "..", "ThirdParty", "USD", "include"));
