USDCore: set camera component FocusMethod=Disable when USD fStop is zero

The UsdGeomCamera schema specifies that when the fStop attibute is zero
(the default value), depth of field effects should be disabled.

diff --git a/Engine/Plugins/Runtime/USDCore/Source/USDUtilities/Private/USDPrimConversion.cpp b/Engine/Plugins/Runtime/USDCore/Source/USDUtilities/Private/USDPrimConversion.cpp
--- a/Engine/Plugins/Runtime/USDCore/Source/USDUtilities/Private/USDPrimConversion.cpp
+++ b/Engine/Plugins/Runtime/USDCore/Source/USDUtilities/Private/USDPrimConversion.cpp
@@ -536,7 +536,13 @@ bool UsdToUnreal::ConvertGeomCamera(const UE::FUsdPrim& Prim, UCineCameraCompone
 
 	if (pxr::UsdAttribute Attr = GeomCamera.GetFStopAttr(); Attr.HasAuthoredValue())
 	{
-		CameraComponent.SetCurrentAperture(UsdUtils::GetUsdValue<float>(Attr, UsdTimeCode));
+		const float UsdFStop = UsdUtils::GetUsdValue<float>(Attr, UsdTimeCode);
+		CameraComponent.SetCurrentAperture(UsdFStop);
+
+		if (FMath::IsNearlyZero(UsdFStop))
+		{
+			CameraComponent.FocusSettings.FocusMethod = ECameraFocusMethod::Disable;
+		}
 	}
 
 	if (pxr::UsdAttribute Attr = GeomCamera.GetHorizontalApertureAttr(); Attr.HasAuthoredValue())
