fix for Mac Editor builds
- shouldn't check Target.Architecture on a multi-architecture platform

#rnx

diff --git a/Engine/Plugins/Runtime/USDCore/Source/UnrealUSDWrapper/UnrealUSDWrapper.Build.cs b/Engine/Plugins/Runtime/USDCore/Source/UnrealUSDWrapper/UnrealUSDWrapper.Build.cs
--- a/Engine/Plugins/Runtime/USDCore/Source/UnrealUSDWrapper/UnrealUSDWrapper.Build.cs
+++ b/Engine/Plugins/Runtime/USDCore/Source/UnrealUSDWrapper/UnrealUSDWrapper.Build.cs
@@ -33,7 +33,7 @@ namespace UnrealBuildTool.Rules
 				PublicDefinitions.Add("ENABLE_USD_DEBUG_PATH=0");
 
 				var EngineDir = Path.GetFullPath(Target.RelativeEnginePath);
-				var PlatformDir = (Target.Architecture == UnrealArch.Arm64 && Target.Platform == UnrealTargetPlatform.Win64) ? "WinArm64" : Target.Platform.ToString();
+				var PlatformDir = (Target.Platform == UnrealTargetPlatform.Win64 && Target.Architecture == UnrealArch.Arm64) ? "WinArm64" : Target.Platform.ToString();
 				var PythonSourceTPSDir = Path.Combine(EngineDir, "Source", "ThirdParty", "Python3", PlatformDir);
 				var PythonBinaryTPSDir = Path.Combine(EngineDir, "Binaries", "ThirdParty", "Python3", PlatformDir);
 				string USDLibsDir = Path.Combine(ModuleDirectory, "..", "ThirdParty", "USD", "lib");
