USD ARM64 libs and dlls.

#rnx

diff --git a/Engine/Plugins/Runtime/USDCore/Source/UnrealUSDWrapper/UnrealUSDWrapper.Build.cs b/Engine/Plugins/Runtime/USDCore/Source/UnrealUSDWrapper/UnrealUSDWrapper.Build.cs
--- a/Engine/Plugins/Runtime/USDCore/Source/UnrealUSDWrapper/UnrealUSDWrapper.Build.cs
+++ b/Engine/Plugins/Runtime/USDCore/Source/UnrealUSDWrapper/UnrealUSDWrapper.Build.cs
@@ -36,7 +36,6 @@ namespace UnrealBuildTool.Rules
 				var PlatformDir = (Target.Platform == UnrealTargetPlatform.Win64 && Target.Architecture == UnrealArch.Arm64) ? "WinArm64" : Target.Platform.ToString();
 				var PythonSourceTPSDir = Path.Combine(EngineDir, "Source", "ThirdParty", "Python3", PlatformDir);
 				var PythonBinaryTPSDir = Path.Combine(EngineDir, "Binaries", "ThirdParty", "Python3", PlatformDir);
-				string USDLibsDir = Path.Combine(ModuleDirectory, "..", "ThirdParty", "USD", "lib");
 
 				if (Target.Platform == UnrealTargetPlatform.Win64)
 				{
@@ -49,6 +48,7 @@ namespace UnrealBuildTool.Rules
 					RuntimeDependencies.Add(Path.Combine("$(TargetOutputDir)", ArchDir, "python311.dll"), Path.Combine(PythonBinaryTPSDir, "python311.dll"));
 
 					// USD
+					string USDLibsDir = Path.Combine(ModuleDirectory, "..", "ThirdParty", "USD", "lib", "Win64", Target.Architecture.WindowsLibDir);
 					PublicSystemIncludePaths.Add(Path.Combine(ModuleDirectory, "..", "ThirdParty", "USD", "include"));
 					PublicSystemLibraryPaths.Add(USDLibsDir);
 					foreach (string UsdLib in Directory.EnumerateFiles(USDLibsDir, "*.lib", SearchOption.AllDirectories))
@@ -56,7 +56,7 @@ namespace UnrealBuildTool.Rules
 
 						PublicAdditionalLibraries.Add(UsdLib);
 					}
-					var USDBinDir = Path.Combine(ModuleDirectory, "..", "ThirdParty", "USD", "bin");
+					var USDBinDir = Path.Combine(ModuleDirectory, "..", "ThirdParty", "USD", "bin", "Win64", Target.Architecture.WindowsLibDir);
 					foreach (string UsdDll in Directory.EnumerateFiles(USDBinDir, "*.dll", SearchOption.AllDirectories))
 					{
 						// We can't delay-load the USD dlls as they contain data and vtables: They need to be next to the executable and implicitly linked
