USD: Allow UsdPreviewSurface material parsing to step through output connections.

diff --git a/Engine/Plugins/Runtime/USDCore/Source/USDUtilities/Private/USDShadeConversion.cpp b/Engine/Plugins/Runtime/USDCore/Source/USDUtilities/Private/USDShadeConversion.cpp
--- a/Engine/Plugins/Runtime/USDCore/Source/USDUtilities/Private/USDShadeConversion.cpp
+++ b/Engine/Plugins/Runtime/USDCore/Source/USDUtilities/Private/USDShadeConversion.cpp
@@ -514,13 +514,15 @@ namespace UE::UsdShadeConversion::Private
 		}
 	}
 
-	bool GetTextureParameterValue(pxr::UsdShadeInput& ShadeInput, TextureGroup Group, UsdToUnreal::FParameterValue& OutValue)
+	// This works for both pxr::UsdShadeInput and pxr::UsdShadeOutput
+	template<typename InputOrOutputType>
+	bool GetTextureParameterValue(InputOrOutputType& InputOrOutput, TextureGroup Group, UsdToUnreal::FParameterValue& OutValue)
 	{
 		FScopedUsdAllocs UsdAllocs;
 
-		const bool bIsNormalInput =
-			(ShadeInput.GetTypeName() == pxr::SdfValueTypeNames->Normal3f || ShadeInput.GetTypeName() == pxr::SdfValueTypeNames->Normal3fArray
-			 || ShadeInput.GetBaseName() == UnrealIdentifiers::Normal);
+		const bool bIsNormalInput = (InputOrOutput.GetTypeName() == pxr::SdfValueTypeNames->Normal3f 
+								  || InputOrOutput.GetTypeName() == pxr::SdfValueTypeNames->Normal3fArray
+			 					  || InputOrOutput.GetBaseName() == UnrealIdentifiers::Normal);
 
 		pxr::UsdShadeConnectableAPI UsdUVTextureSource;
 		pxr::TfToken UsdUVTextureSourceName;
@@ -530,7 +532,7 @@ namespace UE::UsdShadeConversion::Private
 		OutValue.Set<float>(0.0f);
 
 		if (pxr::UsdShadeConnectableAPI::GetConnectedSource(
-				ShadeInput.GetAttr(),
+				InputOrOutput.GetAttr(),
 				&UsdUVTextureSource,
 				&UsdUVTextureSourceName,
 				&UsdUVTextureAttributeType
@@ -706,20 +708,19 @@ namespace UE::UsdShadeConversion::Private
 	{
 		FScopedUsdAllocs Allocs;
 
-		pxr::UsdShadeInput Input = Connectable.GetInput(InputName);
-		if (!Input)
+		auto ProcessPort = [&InputName, &OutValue, &DefaultValue]<typename PortType>(const PortType& Port) -> bool
 		{
-			return false;
-		}
-
-		// If we have another shader/node connected
-		pxr::UsdShadeConnectableAPI Source;
-		pxr::TfToken SourceName;
-		pxr::UsdShadeAttributeType AttributeType;
-		if (pxr::UsdShadeConnectableAPI::GetConnectedSource(Input.GetAttr(), &Source, &SourceName, &AttributeType))
-		{
-			if (!GetTextureParameterValue(Input, TEXTUREGROUP_WorldSpecular, OutValue))
+			// If we have another shader/node connected
+			pxr::UsdShadeConnectableAPI Source;
+			pxr::TfToken SourceName;
+			pxr::UsdShadeAttributeType AttributeType;
+			if (pxr::UsdShadeConnectableAPI::GetConnectedSource(Port.GetAttr(), &Source, &SourceName, &AttributeType))
 			{
+				if (GetTextureParameterValue(Port, TEXTUREGROUP_WorldSpecular, OutValue))
+				{
+					return true;
+				}
+
 				// Check if we have a fallback input that we can use instead, since we don't have a valid texture value
 				if (const pxr::UsdShadeInput FallbackInput = Source.GetInput(UnrealIdentifiers::Fallback))
 				{
@@ -735,21 +736,42 @@ namespace UE::UsdShadeConversion::Private
 				// (e.g. when shader input is just "hoisted" and connected to the parent material input)
 				return GetFloatParameterValue(Source, SourceName, DefaultValue, OutValue);
 			}
-		}
-		// No other node connected, so we must have some value
-		else
+			// No other node connected, so we must have some value
+			else
+			{
+				float InputValue = DefaultValue;
+
+				if constexpr (std::is_same_v<PortType, pxr::UsdShadeInput>)
+				{
+					Port.template Get<float>(&InputValue);
+				}
+
+				OutValue.Set<float>(InputValue);
+				return true;
+			}
+
+			return false;
+		};
+
+		pxr::UsdShadeInput Input = Connectable.GetInput(InputName);
+		if (Input)
 		{
-			float InputValue = DefaultValue;
-			Input.Get<float>(&InputValue);
+			return ProcessPort(Input);
+		}
 
-			OutValue.Set<float>(InputValue);
+		// Check if we're connected to another NodeGraph output instead, and try climbing through that
+		pxr::UsdShadeOutput Output = Connectable.GetOutput(InputName);
+		if (Output)
+		{
+			return ProcessPort(Output);
 		}
 
-		return true;
+		return false;
 	}
 
+	template<typename InputOrOutputType>
 	bool GetPrimvarReaderParameterValue(
-		const pxr::UsdShadeInput& Input,
+		const InputOrOutputType& Input,
 		const pxr::TfToken& PrimvarReaderShaderId,
 		const FLinearColor& DefaultValue,
 		UsdToUnreal::FParameterValue& OutValue
@@ -837,28 +859,27 @@ namespace UE::UsdShadeConversion::Private
 	{
 		FScopedUsdAllocs Allocs;
 
-		pxr::UsdShadeInput Input = Connectable.GetInput(InputName);
-		if (!Input)
-		{
-			return false;
-		}
-
-		// If we have another shader/node connected
-		pxr::UsdShadeConnectableAPI Source;
-		pxr::TfToken SourceName;
-		pxr::UsdShadeAttributeType AttributeType;
-		if (pxr::UsdShadeConnectableAPI::GetConnectedSource(Input.GetAttr(), &Source, &SourceName, &AttributeType))
+		auto ProcessPort = [Group, &InputName, &OutValue, &DefaultValue]<typename PortType>(const PortType& Port) -> bool
 		{
-			if (!GetTextureParameterValue(Input, Group, OutValue))
+			// Check if we're connected to something
+			pxr::UsdShadeConnectableAPI Source;
+			pxr::TfToken SourceName;
+			pxr::UsdShadeAttributeType AttributeType;
+			if (pxr::UsdShadeConnectableAPI::GetConnectedSource(Port.GetAttr(), &Source, &SourceName, &AttributeType))
 			{
-				// Check whether this input receives its value through a connection to a
-				// primvar reader shader.
-				if (GetPrimvarReaderParameterValue(Input, UnrealIdentifiers::UsdPrimvarReader_float3, DefaultValue, OutValue))
+				// Check if we're connected to a UsdUVTexture prim
+				if (GetTextureParameterValue(Port, Group, OutValue))
 				{
 					return true;
 				}
 
-				// Check if we have a fallback input that we can use instead, since we don't have a valid texture value
+				// Check whether this input receives its value through a connection to a primvar reader shader
+				if (GetPrimvarReaderParameterValue(Port, UnrealIdentifiers::UsdPrimvarReader_float3, DefaultValue, OutValue))
+				{
+					return true;
+				}
+
+				// Check if we have a fallback input on the connection, since we don't have a valid texture value
 				if (const pxr::UsdShadeInput FallbackInput = Source.GetInput(UnrealIdentifiers::Fallback))
 				{
 					pxr::GfVec3f UsdFallbackVec3;
@@ -888,24 +909,44 @@ namespace UE::UsdShadeConversion::Private
 					}
 				}
 
-				// This shader doesn't have anything: Traverse into the input connectable itself
+				// We're connected to something else, so try traversing into the input connectable itself
 				return GetVec3ParameterValue(Source, SourceName, DefaultValue, OutValue, Group);
 			}
-		}
-		// No other node connected, so we must have some value
-		else if (InputName != UnrealIdentifiers::Normal)
-		{
-			FLinearColor DiffuseColor = DefaultValue;
-			pxr::GfVec3f UsdDiffuseColor;
-			if (Input.Get<pxr::GfVec3f>(&UsdDiffuseColor))
+			// Has a value directly on the port
+			else if (InputName != UnrealIdentifiers::Normal)
 			{
-				DiffuseColor = UsdToUnreal::ConvertColor(UsdDiffuseColor);
+				FLinearColor DiffuseColor = DefaultValue;
+
+				if constexpr (std::is_same_v<PortType, pxr::UsdShadeInput>)
+				{
+					pxr::GfVec3f UsdDiffuseColor;
+					if (Port.template Get<pxr::GfVec3f>(&UsdDiffuseColor))
+					{
+						DiffuseColor = UsdToUnreal::ConvertColor(UsdDiffuseColor);
+					}
+				}
+
+				OutValue.Set<FVector>(FVector(DiffuseColor));
+				return true;
 			}
 
-			OutValue.Set<FVector>(FVector(DiffuseColor));
+			return false;
+		};
+
+		pxr::UsdShadeInput Input = Connectable.GetInput(InputName);
+		if (Input)
+		{
+			return ProcessPort(Input);
 		}
 
-		return true;
+		// Check if we're connected to another NodeGraph output instead, and try climbing through that
+		pxr::UsdShadeOutput Output = Connectable.GetOutput(InputName);
+		if (Output)
+		{
+			return ProcessPort(Output);
+		}
+
+		return false;
 	}
 
 	bool GetBoolParameterValue(
