USD: Load USD plugins from where the loaded module resides instead of BaseDir.

diff --git a/Engine/Plugins/Runtime/USDCore/Source/UnrealUSDWrapper/Private/UnrealUSDWrapper.cpp b/Engine/Plugins/Runtime/USDCore/Source/UnrealUSDWrapper/Private/UnrealUSDWrapper.cpp
--- a/Engine/Plugins/Runtime/USDCore/Source/UnrealUSDWrapper/Private/UnrealUSDWrapper.cpp
+++ b/Engine/Plugins/Runtime/USDCore/Source/UnrealUSDWrapper/Private/UnrealUSDWrapper.cpp
@@ -1122,7 +1122,28 @@ public:
 #endif	  // PLATFORM_WINDOWS
 
 #else
-		FString TargetDllFolder = FPlatformProcess::BaseDir();
+		FString TargetDllFolder;
+
+		FModuleManager& ModuleManager = FModuleManager::Get();
+#if !IS_MONOLITHIC
+		// Load the USD plugins from where the UnrealUSDWrapper module was loaded since BaseDir
+		// can point to another location entirely and we don't want to load the USD modules twice
+		FString SelfFileName = ModuleManager.GetModuleFilename(TEXT("UnrealUSDWrapper"));
+		if (!SelfFileName.IsEmpty())
+		{
+			TargetDllFolder = FPaths::GetPath(SelfFileName);
+		}
+#endif // IS_MONOLITHIC
+
+		if (TargetDllFolder.IsEmpty())
+		{
+			TargetDllFolder = ModuleManager.GetGameBinariesDirectory();
+		}
+
+		if (TargetDllFolder.IsEmpty())
+		{
+			TargetDllFolder = FPlatformProcess::BaseDir();
+		}
 #endif	  // USE_LIBRARIES_FROM_PLUGIN_FOLDER
 
 		// Path to the MaterialX standard data libraries.
