USD: Add more UsdWrapper types and tweak Unreal.natvis to improve how they're displayed on the debugger.

diff --git a/Engine/Plugins/Runtime/USDCore/Source/UnrealUSDWrapper/Private/UsdWrappers/PcpLayerStack.cpp b/Engine/Plugins/Runtime/USDCore/Source/UnrealUSDWrapper/Private/UsdWrappers/PcpLayerStack.cpp
new file mode 100644
--- /dev/null
+++ b/Engine/Plugins/Runtime/USDCore/Source/UnrealUSDWrapper/Private/UsdWrappers/PcpLayerStack.cpp
@@ -0,0 +1,222 @@
+// Copyright Epic Games, Inc. All Rights Reserved.
+
+#include "UsdWrappers/PcpLayerStack.h"
+
+#include "USDMemory.h"
+#include "UsdWrappers/PcpLayerStackIdentifier.h"
+#include "UsdWrappers/SdfLayer.h"
+
+#if USE_USD_SDK
+#include "USDIncludesStart.h"
+#include <pxr/usd/pcp/layerStack.h>
+#include <pxr/usd/pcp/layerStackIdentifier.h>
+#include "USDIncludesEnd.h"
+#endif	  // #if USE_USD_SDK
+
+namespace UE
+{
+	namespace Internal
+	{
+		class FPcpLayerStackImpl
+		{
+		public:
+			FPcpLayerStackImpl() = default;
+
+#if ENABLE_USD_DEBUG_PATH
+			FPcpLayerStackIdentifier DebugIdentifier;
+#endif
+
+#if USE_USD_SDK
+			explicit FPcpLayerStackImpl(const pxr::PcpLayerStackRefPtr& InPcpLayerStack)
+				: PxrPcpLayerStack(InPcpLayerStack)
+			{
+#if ENABLE_USD_DEBUG_PATH
+				DebugIdentifier = FPcpLayerStackIdentifier{(*PxrPcpLayerStack)->GetIdentifier()};
+#endif
+			}
+
+			explicit FPcpLayerStackImpl(pxr::PcpLayerStackRefPtr&& InPcpLayerStack)
+				: PxrPcpLayerStack(MoveTemp(InPcpLayerStack))
+			{
+#if ENABLE_USD_DEBUG_PATH
+				DebugIdentifier = FPcpLayerStackIdentifier{(*PxrPcpLayerStack)->GetIdentifier()};
+#endif
+			}
+
+			TUsdStore<pxr::PcpLayerStackRefPtr> PxrPcpLayerStack;
+#endif	  // #if USE_USD_SDK
+		};
+	}	 // namespace Internal
+}	 // namespace UE
+
+namespace UE
+{
+	FPcpLayerStack::FPcpLayerStack()
+	{
+		FScopedUnrealAllocs UnrealAllocs;
+		Impl = MakeUnique<Internal::FPcpLayerStackImpl>();
+	}
+
+	FPcpLayerStack::FPcpLayerStack(const FPcpLayerStack& Other)
+	{
+#if USE_USD_SDK
+		FScopedUnrealAllocs UnrealAllocs;
+		Impl = MakeUnique<Internal::FPcpLayerStackImpl>(Other.Impl->PxrPcpLayerStack.Get());
+#endif	  // #if USE_USD_SDK
+	}
+
+	FPcpLayerStack::FPcpLayerStack(FPcpLayerStack&& Other) = default;
+
+	FPcpLayerStack::~FPcpLayerStack()
+	{
+		FScopedUnrealAllocs UnrealAllocs;
+		Impl.Reset();
+	}
+
+	FPcpLayerStack& FPcpLayerStack::operator=(const FPcpLayerStack& Other)
+	{
+#if USE_USD_SDK
+		FScopedUnrealAllocs UnrealAllocs;
+		Impl = MakeUnique<Internal::FPcpLayerStackImpl>(Other.Impl->PxrPcpLayerStack.Get());
+#endif	  // #if USE_USD_SDK
+		return *this;
+	}
+
+	FPcpLayerStack& FPcpLayerStack::operator=(FPcpLayerStack&& Other)
+	{
+		FScopedUnrealAllocs UnrealAllocs;
+		Impl = MoveTemp(Other.Impl);
+
+		return *this;
+	}
+
+	bool FPcpLayerStack::operator==(const FPcpLayerStack& Other) const
+	{
+#if USE_USD_SDK
+		return Impl->PxrPcpLayerStack.Get() == Other.Impl->PxrPcpLayerStack.Get();
+#else
+		return false;
+#endif	  // #if USE_USD_SDK
+	}
+
+	bool FPcpLayerStack::operator!=(const FPcpLayerStack& Other) const
+	{
+		return !(*this == Other);
+	}
+
+	FPcpLayerStack::operator bool() const
+	{
+#if USE_USD_SDK
+		return static_cast<bool>(Impl->PxrPcpLayerStack);
+#else
+		return false;
+#endif	  // #if USE_USD_SDK
+	}
+
+#if USE_USD_SDK
+	FPcpLayerStack::FPcpLayerStack(const pxr::PcpLayerStackRefPtr& InPcpLayerStack)
+	{
+		FScopedUnrealAllocs UnrealAllocs;
+		Impl = MakeUnique<Internal::FPcpLayerStackImpl>(InPcpLayerStack);
+	}
+
+	FPcpLayerStack::FPcpLayerStack(pxr::PcpLayerStackRefPtr&& InPcpLayerStack)
+	{
+		FScopedUnrealAllocs UnrealAllocs;
+		Impl = MakeUnique<Internal::FPcpLayerStackImpl>(MoveTemp(InPcpLayerStack));
+	}
+
+	FPcpLayerStack& FPcpLayerStack::operator=(const pxr::PcpLayerStackRefPtr& InPcpLayerStack)
+	{
+		FScopedUnrealAllocs UnrealAllocs;
+		Impl = MakeUnique<Internal::FPcpLayerStackImpl>(InPcpLayerStack);
+		return *this;
+	}
+
+	FPcpLayerStack& FPcpLayerStack::operator=(pxr::PcpLayerStackRefPtr&& InPcpLayerStack)
+	{
+		FScopedUnrealAllocs UnrealAllocs;
+		Impl = MakeUnique<Internal::FPcpLayerStackImpl>(MoveTemp(InPcpLayerStack));
+		return *this;
+	}
+
+	FPcpLayerStack::operator pxr::PcpLayerStackRefPtr&()
+	{
+		return Impl->PxrPcpLayerStack.Get(); 
+	}
+
+	FPcpLayerStack::operator const pxr::PcpLayerStackRefPtr&() const
+	{
+		return Impl->PxrPcpLayerStack.Get();
+	}
+#endif	  // #if USE_USD_SDK
+
+	FPcpLayerStackIdentifier FPcpLayerStack::GetIdentifier() const
+	{
+		FPcpLayerStackIdentifier Result;
+#if USE_USD_SDK
+		Result = FPcpLayerStackIdentifier{Impl->PxrPcpLayerStack.Get()->GetIdentifier()};
+#endif	  // USE_USD_SDK
+		return Result;
+	}
+
+	TArray<FSdfLayer> FPcpLayerStack::GetLayers() const
+	{
+		TArray<FSdfLayer> Result;
+#if USE_USD_SDK
+		TUsdStore<std::vector<pxr::SdfLayerRefPtr>> Layers = Impl->PxrPcpLayerStack.Get()->GetLayers();
+		Result.Reserve(Layers->size());
+
+		for (const pxr::SdfLayerRefPtr& SdfLayer : *Layers)
+		{
+			Result.Add(FSdfLayer{SdfLayer});
+		}
+#endif	  // USE_USD_SDK
+		return Result;
+	}
+
+	FSdfLayerOffset FPcpLayerStack::GetLayerOffsetForLayer(const FSdfLayer& Layer) const
+	{
+		FSdfLayerOffset Result;
+#if USE_USD_SDK
+		if (const pxr::SdfLayerOffset* UsdOffest = Impl->PxrPcpLayerStack.Get()->GetLayerOffsetForLayer(pxr::SdfLayerRefPtr{Layer}))
+		{
+			Result.Offset = UsdOffest->GetOffset();
+			Result.Scale = UsdOffest->GetScale();
+		}
+#endif	  // USE_USD_SDK
+		return Result;
+	}
+
+	FSdfLayerOffset FPcpLayerStack::GetLayerOffsetForLayer(const FSdfLayerWeak& Layer) const
+	{
+		FSdfLayerOffset Result;
+#if USE_USD_SDK
+		if (const pxr::SdfLayerOffset* UsdOffest = Impl->PxrPcpLayerStack.Get()->GetLayerOffsetForLayer(pxr::SdfLayerHandle{Layer}))
+		{
+			Result.Offset = UsdOffest->GetOffset();
+			Result.Scale = UsdOffest->GetScale();
+		}
+#endif	  // USE_USD_SDK
+		return Result;
+	}
+
+	bool FPcpLayerStack::HasLayer(const FSdfLayer& Layer) const
+	{
+#if USE_USD_SDK
+		return Impl->PxrPcpLayerStack.Get()->HasLayer(pxr::SdfLayerRefPtr{Layer});
+#else
+		return false;
+#endif	  // USE_USD_SDK
+	}
+
+	bool FPcpLayerStack::HasLayer(const FSdfLayerWeak& Layer) const
+	{
+#if USE_USD_SDK
+		return Impl->PxrPcpLayerStack.Get()->HasLayer(pxr::SdfLayerHandle{Layer});
+#else
+		return false;
+#endif	  // USE_USD_SDK
+	}
+
+}	 // namespace UE
diff --git a/Engine/Plugins/Runtime/USDCore/Source/UnrealUSDWrapper/Private/UsdWrappers/PcpLayerStackIdentifier.cpp b/Engine/Plugins/Runtime/USDCore/Source/UnrealUSDWrapper/Private/UsdWrappers/PcpLayerStackIdentifier.cpp
new file mode 100644
--- /dev/null
+++ b/Engine/Plugins/Runtime/USDCore/Source/UnrealUSDWrapper/Private/UsdWrappers/PcpLayerStackIdentifier.cpp
@@ -0,0 +1,185 @@
+// Copyright Epic Games, Inc. All Rights Reserved.
+
+#include "UsdWrappers/PcpLayerStackIdentifier.h"
+
+#include "UsdWrappers/SdfLayer.h"
+#include "USDMemory.h"
+
+#if USE_USD_SDK
+#include "USDIncludesStart.h"
+#include <pxr/usd/pcp/layerStack.h>
+#include "USDIncludesEnd.h"
+#endif	  // #if USE_USD_SDK
+
+namespace UE
+{
+	namespace Internal
+	{
+		class FPcpLayerStackIdentifierImpl
+		{
+		public:
+			FPcpLayerStackIdentifierImpl() = default;
+
+#if ENABLE_USD_DEBUG_PATH
+			FString DebugRootLayerIdentifier;
+			FString DebugSessionLayerIdentifier;
+#endif
+
+#if USE_USD_SDK
+			explicit FPcpLayerStackIdentifierImpl(const pxr::PcpLayerStackIdentifier& InPcpLayerStackIdentifier)
+				: PxrPcpLayerStackIdentifier(InPcpLayerStackIdentifier)
+			{
+#if ENABLE_USD_DEBUG_PATH
+				if (PxrPcpLayerStackIdentifier->rootLayer)
+				{
+					DebugRootLayerIdentifier = FString{UTF8_TO_TCHAR(PxrPcpLayerStackIdentifier->rootLayer->GetIdentifier().c_str())};
+				}
+				if (PxrPcpLayerStackIdentifier->sessionLayer)
+				{
+					DebugSessionLayerIdentifier = FString{UTF8_TO_TCHAR(PxrPcpLayerStackIdentifier->sessionLayer->GetIdentifier().c_str())};
+				}
+#endif
+			}
+
+			explicit FPcpLayerStackIdentifierImpl(pxr::PcpLayerStackIdentifier&& InPcpLayerStackIdentifier)
+				: PxrPcpLayerStackIdentifier(MoveTemp(InPcpLayerStackIdentifier))
+			{
+#if ENABLE_USD_DEBUG_PATH
+				if (PxrPcpLayerStackIdentifier->rootLayer)
+				{
+					DebugRootLayerIdentifier = FString{UTF8_TO_TCHAR(PxrPcpLayerStackIdentifier->rootLayer->GetIdentifier().c_str())};
+				}
+				if (PxrPcpLayerStackIdentifier->sessionLayer)
+				{
+					DebugSessionLayerIdentifier = FString{UTF8_TO_TCHAR(PxrPcpLayerStackIdentifier->sessionLayer->GetIdentifier().c_str())};
+				}
+#endif
+			}
+
+			TUsdStore<pxr::PcpLayerStackIdentifier> PxrPcpLayerStackIdentifier;
+#endif	  // #if USE_USD_SDK
+		};
+	}	 // namespace Internal
+}	 // namespace UE
+
+namespace UE
+{
+	FPcpLayerStackIdentifier::FPcpLayerStackIdentifier()
+	{
+		FScopedUnrealAllocs UnrealAllocs;
+		Impl = MakeUnique<Internal::FPcpLayerStackIdentifierImpl>();
+	}
+
+	FPcpLayerStackIdentifier::FPcpLayerStackIdentifier(const FPcpLayerStackIdentifier& Other)
+	{
+#if USE_USD_SDK
+		FScopedUnrealAllocs UnrealAllocs;
+		Impl = MakeUnique<Internal::FPcpLayerStackIdentifierImpl>(Other.Impl->PxrPcpLayerStackIdentifier.Get());
+#endif	  // #if USE_USD_SDK
+	}
+
+	FPcpLayerStackIdentifier::FPcpLayerStackIdentifier(FPcpLayerStackIdentifier&& Other) = default;
+
+	FPcpLayerStackIdentifier::~FPcpLayerStackIdentifier()
+	{
+		FScopedUnrealAllocs UnrealAllocs;
+		Impl.Reset();
+	}
+
+	FPcpLayerStackIdentifier& FPcpLayerStackIdentifier::operator=(const FPcpLayerStackIdentifier& Other)
+	{
+#if USE_USD_SDK
+		FScopedUnrealAllocs UnrealAllocs;
+		Impl = MakeUnique<Internal::FPcpLayerStackIdentifierImpl>(Other.Impl->PxrPcpLayerStackIdentifier.Get());
+#endif	  // #if USE_USD_SDK
+		return *this;
+	}
+
+	FPcpLayerStackIdentifier& FPcpLayerStackIdentifier::operator=(FPcpLayerStackIdentifier&& Other)
+	{
+		FScopedUnrealAllocs UnrealAllocs;
+		Impl = MoveTemp(Other.Impl);
+
+		return *this;
+	}
+
+	bool FPcpLayerStackIdentifier::operator==(const FPcpLayerStackIdentifier& Other) const
+	{
+#if USE_USD_SDK
+		return Impl->PxrPcpLayerStackIdentifier.Get() == Other.Impl->PxrPcpLayerStackIdentifier.Get();
+#else
+		return false;
+#endif	  // #if USE_USD_SDK
+	}
+
+	bool FPcpLayerStackIdentifier::operator!=(const FPcpLayerStackIdentifier& Other) const
+	{
+		return !(*this == Other);
+	}
+
+	FPcpLayerStackIdentifier::operator bool() const
+	{
+#if USE_USD_SDK
+		return static_cast<bool>(Impl->PxrPcpLayerStackIdentifier);
+#else
+		return false;
+#endif	  // #if USE_USD_SDK
+	}
+
+#if USE_USD_SDK
+	FPcpLayerStackIdentifier::FPcpLayerStackIdentifier(const pxr::PcpLayerStackIdentifier& InPcpLayerStackIdentifier)
+	{
+		FScopedUnrealAllocs UnrealAllocs;
+		Impl = MakeUnique<Internal::FPcpLayerStackIdentifierImpl>(InPcpLayerStackIdentifier);
+	}
+
+	FPcpLayerStackIdentifier::FPcpLayerStackIdentifier(pxr::PcpLayerStackIdentifier&& InPcpLayerStackIdentifier)
+	{
+		FScopedUnrealAllocs UnrealAllocs;
+		Impl = MakeUnique<Internal::FPcpLayerStackIdentifierImpl>(MoveTemp(InPcpLayerStackIdentifier));
+	}
+
+	FPcpLayerStackIdentifier& FPcpLayerStackIdentifier::operator=(const pxr::PcpLayerStackIdentifier& InPcpLayerStackIdentifier)
+	{
+		FScopedUnrealAllocs UnrealAllocs;
+		Impl = MakeUnique<Internal::FPcpLayerStackIdentifierImpl>(InPcpLayerStackIdentifier);
+		return *this;
+	}
+
+	FPcpLayerStackIdentifier& FPcpLayerStackIdentifier::operator=(pxr::PcpLayerStackIdentifier&& InPcpLayerStackIdentifier)
+	{
+		FScopedUnrealAllocs UnrealAllocs;
+		Impl = MakeUnique<Internal::FPcpLayerStackIdentifierImpl>(MoveTemp(InPcpLayerStackIdentifier));
+		return *this;
+	}
+
+	FPcpLayerStackIdentifier::operator pxr::PcpLayerStackIdentifier&()
+	{
+		return Impl->PxrPcpLayerStackIdentifier.Get(); 
+	}
+
+	FPcpLayerStackIdentifier::operator const pxr::PcpLayerStackIdentifier&() const
+	{
+		return Impl->PxrPcpLayerStackIdentifier.Get();
+	}
+#endif	  // #if USE_USD_SDK
+
+	FSdfLayerWeak FPcpLayerStackIdentifier::RootLayer() const
+	{
+		FSdfLayerWeak Result;
+#if USE_USD_SDK
+		Result = FSdfLayerWeak{Impl->PxrPcpLayerStackIdentifier.Get().rootLayer};
+#endif	  // USE_USD_SDK
+		return Result;
+	}
+
+	FSdfLayerWeak FPcpLayerStackIdentifier::SessionLayer() const
+	{
+		FSdfLayerWeak Result;
+#if USE_USD_SDK
+		Result = FSdfLayerWeak{Impl->PxrPcpLayerStackIdentifier.Get().sessionLayer};
+#endif	  // USE_USD_SDK
+		return Result;
+	}
+
+}	 // namespace UE
diff --git a/Engine/Plugins/Runtime/USDCore/Source/UnrealUSDWrapper/Private/UsdWrappers/PcpMapExpression.cpp b/Engine/Plugins/Runtime/USDCore/Source/UnrealUSDWrapper/Private/UsdWrappers/PcpMapExpression.cpp
new file mode 100644
--- /dev/null
+++ b/Engine/Plugins/Runtime/USDCore/Source/UnrealUSDWrapper/Private/UsdWrappers/PcpMapExpression.cpp
@@ -0,0 +1,143 @@
+// Copyright Epic Games, Inc. All Rights Reserved.
+
+#include "UsdWrappers/PcpMapExpression.h"
+
+#include "UsdWrappers/SdfLayer.h"
+#include "USDMemory.h"
+
+#if USE_USD_SDK
+#include "USDIncludesStart.h"
+#include "pxr/usd/pcp/mapExpression.h"
+#include "USDIncludesEnd.h"
+#endif	  // #if USE_USD_SDK
+
+namespace UE
+{
+	namespace Internal
+	{
+		class FPcpMapExpressionImpl
+		{
+		public:
+			FPcpMapExpressionImpl() = default;
+
+#if ENABLE_USD_DEBUG_PATH
+			FSdfLayerOffset LayerOffset;
+#endif
+
+#if USE_USD_SDK
+			explicit FPcpMapExpressionImpl(const pxr::PcpMapExpression& InPcpMapExpression)
+				: PxrPcpMapExpression(InPcpMapExpression)
+			{
+#if ENABLE_USD_DEBUG_PATH && USE_USD_SDK
+				const pxr::SdfLayerOffset& PcpMapToParent = PxrPcpMapExpression->GetTimeOffset();
+				LayerOffset = FSdfLayerOffset{PcpMapToParent.GetOffset(), PcpMapToParent.GetScale()};
+#endif
+			}
+
+			explicit FPcpMapExpressionImpl(pxr::PcpMapExpression&& InPcpMapExpression)
+				: PxrPcpMapExpression(MoveTemp(InPcpMapExpression))
+			{
+#if ENABLE_USD_DEBUG_PATH && USE_USD_SDK
+				const pxr::SdfLayerOffset& PcpMapToParent = PxrPcpMapExpression->GetTimeOffset();
+				LayerOffset = FSdfLayerOffset{PcpMapToParent.GetOffset(), PcpMapToParent.GetScale()};
+#endif
+			}
+
+			TUsdStore<pxr::PcpMapExpression> PxrPcpMapExpression;
+#endif	  // #if USE_USD_SDK
+		};
+	}	 // namespace Internal
+}	 // namespace UE
+
+namespace UE
+{
+	FPcpMapExpression::FPcpMapExpression()
+	{
+		FScopedUnrealAllocs UnrealAllocs;
+		Impl = MakeUnique<Internal::FPcpMapExpressionImpl>();
+	}
+
+	FPcpMapExpression::FPcpMapExpression(const FPcpMapExpression& Other)
+	{
+#if USE_USD_SDK
+		FScopedUnrealAllocs UnrealAllocs;
+		Impl = MakeUnique<Internal::FPcpMapExpressionImpl>(Other.Impl->PxrPcpMapExpression.Get());
+#endif	  // #if USE_USD_SDK
+	}
+
+	FPcpMapExpression::FPcpMapExpression(FPcpMapExpression&& Other) = default;
+
+	FPcpMapExpression::~FPcpMapExpression()
+	{
+		FScopedUnrealAllocs UnrealAllocs;
+		Impl.Reset();
+	}
+
+	FPcpMapExpression& FPcpMapExpression::operator=(const FPcpMapExpression& Other)
+	{
+#if USE_USD_SDK
+		FScopedUnrealAllocs UnrealAllocs;
+		Impl = MakeUnique<Internal::FPcpMapExpressionImpl>(Other.Impl->PxrPcpMapExpression.Get());
+#endif	  // #if USE_USD_SDK
+		return *this;
+	}
+
+	FPcpMapExpression& FPcpMapExpression::operator=(FPcpMapExpression&& Other)
+	{
+		FScopedUnrealAllocs UnrealAllocs;
+		Impl = MoveTemp(Other.Impl);
+
+		return *this;
+	}
+
+#if USE_USD_SDK
+	FPcpMapExpression::FPcpMapExpression(const pxr::PcpMapExpression& InPcpMapExpression)
+	{
+		FScopedUnrealAllocs UnrealAllocs;
+		Impl = MakeUnique<Internal::FPcpMapExpressionImpl>(InPcpMapExpression);
+	}
+
+	FPcpMapExpression::FPcpMapExpression(pxr::PcpMapExpression&& InPcpMapExpression)
+	{
+		FScopedUnrealAllocs UnrealAllocs;
+		Impl = MakeUnique<Internal::FPcpMapExpressionImpl>(MoveTemp(InPcpMapExpression));
+	}
+
+	FPcpMapExpression& FPcpMapExpression::operator=(const pxr::PcpMapExpression& InPcpMapExpression)
+	{
+		FScopedUnrealAllocs UnrealAllocs;
+		Impl = MakeUnique<Internal::FPcpMapExpressionImpl>(InPcpMapExpression);
+		return *this;
+	}
+
+	FPcpMapExpression& FPcpMapExpression::operator=(pxr::PcpMapExpression&& InPcpMapExpression)
+	{
+		FScopedUnrealAllocs UnrealAllocs;
+		Impl = MakeUnique<Internal::FPcpMapExpressionImpl>(MoveTemp(InPcpMapExpression));
+		return *this;
+	}
+
+	FPcpMapExpression::operator pxr::PcpMapExpression&()
+	{
+		return Impl->PxrPcpMapExpression.Get();
+	}
+
+	FPcpMapExpression::operator const pxr::PcpMapExpression&() const
+	{
+		return Impl->PxrPcpMapExpression.Get();
+	}
+#endif	  // #if USE_USD_SDK
+
+	FSdfLayerOffset FPcpMapExpression::GetTimeOffset() const
+	{
+		FSdfLayerOffset Result;
+
+#if USE_USD_SDK
+		const pxr::SdfLayerOffset& UsdOffset = Impl->PxrPcpMapExpression.Get().GetTimeOffset();
+		Result.Offset = UsdOffset.GetOffset();
+		Result.Scale = UsdOffset.GetScale();
+#endif	  // #if USE_USD_SDK
+
+		return Result;
+	}
+}	 // namespace UE
diff --git a/Engine/Plugins/Runtime/USDCore/Source/UnrealUSDWrapper/Private/UsdWrappers/PcpNodeRef.cpp b/Engine/Plugins/Runtime/USDCore/Source/UnrealUSDWrapper/Private/UsdWrappers/PcpNodeRef.cpp
new file mode 100644
--- /dev/null
+++ b/Engine/Plugins/Runtime/USDCore/Source/UnrealUSDWrapper/Private/UsdWrappers/PcpNodeRef.cpp
@@ -0,0 +1,253 @@
+// Copyright Epic Games, Inc. All Rights Reserved.
+
+#include "UsdWrappers/PcpNodeRef.h"
+
+#include "UsdWrappers/PcpMapExpression.h"
+#include "UsdWrappers/SdfLayer.h"
+#include "UsdWrappers/SdfPath.h"
+#include "UsdWrappers/PcpLayerStack.h"
+#include "USDMemory.h"
+
+#if USE_USD_SDK
+#include "USDIncludesStart.h"
+#include "pxr/usd/usd/attribute.h"
+#include "pxr/usd/pcp/mapExpression.h"
+#include "pxr/usd/pcp/layerStack.h"
+#include "USDIncludesEnd.h"
+#endif	  // #if USE_USD_SDK
+
+namespace UE
+{
+	namespace Internal
+	{
+		class FPcpNodeRefImpl
+		{
+		public:
+			FPcpNodeRefImpl() = default;
+
+#if USE_USD_SDK
+
+#if ENABLE_USD_DEBUG_PATH
+			FString DebugPath;
+			EPcpArcType DebugArcType;
+			FPcpLayerStack DebugLayerStack;
+			FPcpMapExpression DebugMapToParent;
+			FPcpMapExpression DebugMapToRoot;
+			bool bHasSpecs = false;
+			bool bCanContributeSpecs = false;
+			TOptional<FPcpNodeRef> DebugParent; // Optional here otherwise our default constructor would recurse forever
+#endif
+
+			explicit FPcpNodeRefImpl(const pxr::PcpNodeRef& InPcpNodeRef)
+				: PxrPcpNodeRef(InPcpNodeRef)
+			{
+#if ENABLE_USD_DEBUG_PATH
+				SetupDebugMembers();
+#endif
+			}
+
+			explicit FPcpNodeRefImpl(pxr::PcpNodeRef&& InPcpNodeRef)
+				: PxrPcpNodeRef(MoveTemp(InPcpNodeRef))
+			{
+#if ENABLE_USD_DEBUG_PATH
+				SetupDebugMembers();
+#endif
+			}
+
+#if ENABLE_USD_DEBUG_PATH
+			void SetupDebugMembers()
+			{
+				if (!PxrPcpNodeRef || !*PxrPcpNodeRef)
+				{
+					return;
+				}
+
+				DebugPath = FString{UTF8_TO_TCHAR(PxrPcpNodeRef->GetPath().GetString().c_str())};
+				DebugArcType = static_cast<EPcpArcType>(PxrPcpNodeRef->GetArcType());
+				DebugLayerStack = FPcpLayerStack{PxrPcpNodeRef->GetLayerStack()};
+				DebugMapToParent = FPcpMapExpression{PxrPcpNodeRef->GetMapToParent()};
+				DebugMapToRoot = FPcpMapExpression{PxrPcpNodeRef->GetMapToRoot()};
+				bHasSpecs = PxrPcpNodeRef->HasSpecs();
+				bCanContributeSpecs = PxrPcpNodeRef->CanContributeSpecs();
+				if (pxr::PcpNodeRef ParentNode = PxrPcpNodeRef->GetParentNode())
+				{
+					DebugParent = FPcpNodeRef{ParentNode};
+				}
+			}
+#endif
+
+			TUsdStore<pxr::PcpNodeRef> PxrPcpNodeRef;
+#endif	  // #if USE_USD_SDK
+		};
+	}	 // namespace Internal
+}	 // namespace UE
+
+namespace UE
+{
+	FPcpNodeRef::FPcpNodeRef()
+	{
+		FScopedUnrealAllocs UnrealAllocs;
+		Impl = MakeUnique<Internal::FPcpNodeRefImpl>();
+	}
+
+	FPcpNodeRef::FPcpNodeRef(const FPcpNodeRef& Other)
+	{
+#if USE_USD_SDK
+		FScopedUnrealAllocs UnrealAllocs;
+		Impl = MakeUnique<Internal::FPcpNodeRefImpl>(Other.Impl->PxrPcpNodeRef.Get());
+#endif	  // #if USE_USD_SDK
+	}
+
+	FPcpNodeRef::FPcpNodeRef(FPcpNodeRef&& Other) = default;
+
+	FPcpNodeRef::~FPcpNodeRef()
+	{
+		FScopedUnrealAllocs UnrealAllocs;
+		Impl.Reset();
+	}
+
+	FPcpNodeRef& FPcpNodeRef::operator=(const FPcpNodeRef& Other)
+	{
+#if USE_USD_SDK
+		FScopedUnrealAllocs UnrealAllocs;
+		Impl = MakeUnique<Internal::FPcpNodeRefImpl>(Other.Impl->PxrPcpNodeRef.Get());
+#endif	  // #if USE_USD_SDK
+		return *this;
+	}
+
+	FPcpNodeRef& FPcpNodeRef::operator=(FPcpNodeRef&& Other)
+	{
+		FScopedUnrealAllocs UnrealAllocs;
+		Impl = MoveTemp(Other.Impl);
+
+		return *this;
+	}
+
+	bool FPcpNodeRef::operator==(const FPcpNodeRef& Other) const
+	{
+#if USE_USD_SDK
+		return Impl->PxrPcpNodeRef.Get() == Other.Impl->PxrPcpNodeRef.Get();
+#else
+		return false;
+#endif	  // #if USE_USD_SDK
+	}
+
+	bool FPcpNodeRef::operator!=(const FPcpNodeRef& Other) const
+	{
+		return !(*this == Other);
+	}
+
+	FPcpNodeRef::operator bool() const
+	{
+#if USE_USD_SDK
+		return (bool)Impl->PxrPcpNodeRef.Get();
+#else
+		return false;
+#endif	  // #if USE_USD_SDK
+	}
+
+#if USE_USD_SDK
+	FPcpNodeRef::FPcpNodeRef(const pxr::PcpNodeRef& InPcpNodeRef)
+	{
+		FScopedUnrealAllocs UnrealAllocs;
+		Impl = MakeUnique<Internal::FPcpNodeRefImpl>(InPcpNodeRef);
+	}
+
+	FPcpNodeRef::FPcpNodeRef(pxr::PcpNodeRef&& InPcpNodeRef)
+	{
+		FScopedUnrealAllocs UnrealAllocs;
+		Impl = MakeUnique<Internal::FPcpNodeRefImpl>(MoveTemp(InPcpNodeRef));
+	}
+
+	FPcpNodeRef& FPcpNodeRef::operator=(const pxr::PcpNodeRef& InPcpNodeRef)
+	{
+		FScopedUnrealAllocs UnrealAllocs;
+		Impl = MakeUnique<Internal::FPcpNodeRefImpl>(InPcpNodeRef);
+		return *this;
+	}
+
+	FPcpNodeRef& FPcpNodeRef::operator=(pxr::PcpNodeRef&& InPcpNodeRef)
+	{
+		FScopedUnrealAllocs UnrealAllocs;
+		Impl = MakeUnique<Internal::FPcpNodeRefImpl>(MoveTemp(InPcpNodeRef));
+		return *this;
+	}
+
+	FPcpNodeRef::operator pxr::PcpNodeRef&()
+	{
+		return Impl->PxrPcpNodeRef.Get();
+	}
+
+	FPcpNodeRef::operator const pxr::PcpNodeRef&() const
+	{
+		return Impl->PxrPcpNodeRef.Get();
+	}
+#endif	  // #if USE_USD_SDK
+
+	EPcpArcType FPcpNodeRef::GetArcType() const
+	{
+#if USE_USD_SDK
+		static_assert((int)EPcpArcType::PcpArcTypeRoot == (int)pxr::PcpArcType::PcpArcTypeRoot);
+		static_assert((int)EPcpArcType::PcpArcTypeInherit == (int)pxr::PcpArcType::PcpArcTypeInherit);
+		static_assert((int)EPcpArcType::PcpArcTypeVariant == (int)pxr::PcpArcType::PcpArcTypeVariant);
+		static_assert((int)EPcpArcType::PcpArcTypeRelocate == (int)pxr::PcpArcType::PcpArcTypeRelocate);
+		static_assert((int)EPcpArcType::PcpArcTypeReference == (int)pxr::PcpArcType::PcpArcTypeReference);
+		static_assert((int)EPcpArcType::PcpArcTypePayload == (int)pxr::PcpArcType::PcpArcTypePayload);
+		static_assert((int)EPcpArcType::PcpArcTypeSpecialize == (int)pxr::PcpArcType::PcpArcTypeSpecialize);
+		static_assert((int)EPcpArcType::PcpNumArcTypes == (int)pxr::PcpArcType::PcpNumArcTypes);
+
+		return static_cast<EPcpArcType>(Impl->PxrPcpNodeRef.Get().GetArcType());
+#else
+		return {};
+#endif	  // #if USE_USD_SDK
+	}
+
+	FPcpNodeRef FPcpNodeRef::GetParentNode() const
+	{
+		FPcpNodeRef Result;
+
+#if USE_USD_SDK
+		Result = FPcpNodeRef{Impl->PxrPcpNodeRef.Get().GetParentNode()};
+#endif	  // #if USE_USD_SDK
+
+		return Result;
+	}
+
+	FPcpMapExpression FPcpNodeRef::GetMapToParent() const
+	{
+#if USE_USD_SDK
+		return FPcpMapExpression{Impl->PxrPcpNodeRef.Get().GetMapToParent()};
+#else
+		return {};
+#endif	  // #if USE_USD_SDK
+	}
+
+	FPcpMapExpression FPcpNodeRef::GetMapToRoot() const
+	{
+#if USE_USD_SDK
+		return FPcpMapExpression{Impl->PxrPcpNodeRef.Get().GetMapToRoot()};
+#else
+		return {};
+#endif	  // #if USE_USD_SDK
+	}
+
+	FSdfPath FPcpNodeRef::GetPath() const
+	{
+		FSdfPath Path;
+
+#if USE_USD_SDK
+		Path = FSdfPath{Impl->PxrPcpNodeRef.Get().GetPath()};
+#endif	  // #if USE_USD_SDK
+
+		return Path;
+	}
+
+	FPcpLayerStack FPcpNodeRef::GetLayerStack() const
+	{
+#if USE_USD_SDK
+		return FPcpLayerStack{Impl->PxrPcpNodeRef.Get().GetLayerStack()};
+#else
+		return {};
+#endif	  // #if USE_USD_SDK
+	}
+}	 // namespace UE
diff --git a/Engine/Plugins/Runtime/USDCore/Source/UnrealUSDWrapper/Private/UsdWrappers/PcpPrimIndex.cpp b/Engine/Plugins/Runtime/USDCore/Source/UnrealUSDWrapper/Private/UsdWrappers/PcpPrimIndex.cpp
new file mode 100644
--- /dev/null
+++ b/Engine/Plugins/Runtime/USDCore/Source/UnrealUSDWrapper/Private/UsdWrappers/PcpPrimIndex.cpp
@@ -0,0 +1,176 @@
+// Copyright Epic Games, Inc. All Rights Reserved.
+
+#include "UsdWrappers/PcpPrimIndex.h"
+
+#include "UsdWrappers/PcpNodeRef.h"
+#include "UsdWrappers/SdfPath.h"
+#include "USDMemory.h"
+
+#if USE_USD_SDK
+#include "USDIncludesStart.h"
+#include "pxr/usd/pcp/primIndex.h"
+#include "pxr/usd/pcp/node.h"
+#include "USDIncludesEnd.h"
+#endif	  // #if USE_USD_SDK
+
+namespace UE
+{
+	namespace Internal
+	{
+		class FPcpPrimIndexImpl
+		{
+		public:
+			FPcpPrimIndexImpl() = default;
+
+#if USE_USD_SDK
+
+#if ENABLE_USD_DEBUG_PATH
+			FString DebugPath;
+#endif
+
+			explicit FPcpPrimIndexImpl(const pxr::PcpPrimIndex& InPcpPrimIndex)
+				: PxrPcpPrimIndex(InPcpPrimIndex)
+			{
+#if ENABLE_USD_DEBUG_PATH
+				SetupDebugMembers();
+#endif
+			}
+
+			explicit FPcpPrimIndexImpl(pxr::PcpPrimIndex&& InPcpPrimIndex)
+				: PxrPcpPrimIndex(MoveTemp(InPcpPrimIndex))
+			{
+#if ENABLE_USD_DEBUG_PATH
+				SetupDebugMembers();
+#endif
+			}
+
+#if ENABLE_USD_DEBUG_PATH
+			void SetupDebugMembers()
+			{
+				if (!PxrPcpPrimIndex)
+				{
+					return;
+				}
+
+				DebugPath = FString{UTF8_TO_TCHAR(PxrPcpPrimIndex->GetPath().GetString().c_str())};
+			}
+#endif
+
+			TUsdStore<pxr::PcpPrimIndex> PxrPcpPrimIndex;
+#endif	  // #if USE_USD_SDK
+		};
+	}	 // namespace Internal
+}	 // namespace UE
+
+namespace UE
+{
+	FPcpPrimIndex::FPcpPrimIndex()
+	{
+		FScopedUnrealAllocs UnrealAllocs;
+		Impl = MakeUnique<Internal::FPcpPrimIndexImpl>();
+	}
+
+	FPcpPrimIndex::FPcpPrimIndex(const FPcpPrimIndex& Other)
+	{
+#if USE_USD_SDK
+		FScopedUnrealAllocs UnrealAllocs;
+		Impl = MakeUnique<Internal::FPcpPrimIndexImpl>(Other.Impl->PxrPcpPrimIndex.Get());
+#endif	  // #if USE_USD_SDK
+	}
+
+	FPcpPrimIndex::FPcpPrimIndex(FPcpPrimIndex&& Other) = default;
+
+	FPcpPrimIndex::~FPcpPrimIndex()
+	{
+		FScopedUnrealAllocs UnrealAllocs;
+		Impl.Reset();
+	}
+
+	FPcpPrimIndex& FPcpPrimIndex::operator=(const FPcpPrimIndex& Other)
+	{
+#if USE_USD_SDK
+		FScopedUnrealAllocs UnrealAllocs;
+		Impl = MakeUnique<Internal::FPcpPrimIndexImpl>(Other.Impl->PxrPcpPrimIndex.Get());
+#endif	  // #if USE_USD_SDK
+		return *this;
+	}
+
+	FPcpPrimIndex& FPcpPrimIndex::operator=(FPcpPrimIndex&& Other)
+	{
+		FScopedUnrealAllocs UnrealAllocs;
+		Impl = MoveTemp(Other.Impl);
+
+		return *this;
+	}
+
+	FPcpPrimIndex::operator bool() const
+	{
+		return IsValid();
+	}
+
+#if USE_USD_SDK
+	FPcpPrimIndex::FPcpPrimIndex(const pxr::PcpPrimIndex& InPcpPrimIndex)
+	{
+		FScopedUnrealAllocs UnrealAllocs;
+		Impl = MakeUnique<Internal::FPcpPrimIndexImpl>(InPcpPrimIndex);
+	}
+
+	FPcpPrimIndex::FPcpPrimIndex(pxr::PcpPrimIndex&& InPcpPrimIndex)
+	{
+		FScopedUnrealAllocs UnrealAllocs;
+		Impl = MakeUnique<Internal::FPcpPrimIndexImpl>(MoveTemp(InPcpPrimIndex));
+	}
+
+	FPcpPrimIndex& FPcpPrimIndex::operator=(const pxr::PcpPrimIndex& InPcpPrimIndex)
+	{
+		FScopedUnrealAllocs UnrealAllocs;
+		Impl = MakeUnique<Internal::FPcpPrimIndexImpl>(InPcpPrimIndex);
+		return *this;
+	}
+
+	FPcpPrimIndex& FPcpPrimIndex::operator=(pxr::PcpPrimIndex&& InPcpPrimIndex)
+	{
+		FScopedUnrealAllocs UnrealAllocs;
+		Impl = MakeUnique<Internal::FPcpPrimIndexImpl>(MoveTemp(InPcpPrimIndex));
+		return *this;
+	}
+
+	FPcpPrimIndex::operator pxr::PcpPrimIndex&()
+	{
+		return Impl->PxrPcpPrimIndex.Get();
+	}
+
+	FPcpPrimIndex::operator const pxr::PcpPrimIndex&() const
+	{
+		return Impl->PxrPcpPrimIndex.Get();
+	}
+#endif	  // #if USE_USD_SDK
+
+	bool FPcpPrimIndex::IsValid() const
+	{
+#if USE_USD_SDK
+		return (bool)Impl->PxrPcpPrimIndex->IsValid();
+#else
+		return false;
+#endif	  // #if USE_USD_SDK
+	}
+
+	UE::FPcpNodeRef FPcpPrimIndex::GetRootNode() const
+	{
+#if USE_USD_SDK
+		return UE::FPcpNodeRef{Impl->PxrPcpPrimIndex->GetRootNode()};
+#else
+		return {};
+#endif	  // #if USE_USD_SDK
+	}
+
+	FString FPcpPrimIndex::DumpToString(bool bIncludeInheritOriginInfo, bool bIncludeMaps) const
+	{
+#if USE_USD_SDK
+		TUsdStore<std::string> UsdString = Impl->PxrPcpPrimIndex->DumpToString(bIncludeInheritOriginInfo, bIncludeMaps);
+		return FString{UTF8_TO_TCHAR(UsdString->c_str())};
+#else
+		return {};
+#endif	  // #if USE_USD_SDK
+	}
+}	 // namespace UE
diff --git a/Engine/Plugins/Runtime/USDCore/Source/UnrealUSDWrapper/Private/UsdWrappers/SdfLayer.cpp b/Engine/Plugins/Runtime/USDCore/Source/UnrealUSDWrapper/Private/UsdWrappers/SdfLayer.cpp
--- a/Engine/Plugins/Runtime/USDCore/Source/UnrealUSDWrapper/Private/UsdWrappers/SdfLayer.cpp
+++ b/Engine/Plugins/Runtime/USDCore/Source/UnrealUSDWrapper/Private/UsdWrappers/SdfLayer.cpp
@@ -19,20 +19,49 @@ namespace UE
 {
 	namespace Internal
 	{
+#if ENABLE_USD_DEBUG_PATH
+		// The only purpose of this class is so that we can cast the UsdStage impl to this type in the Unreal.natvis debug
+		// visualizers, allowing us to inspect these FUsdStage and FUsdStageWeak variables in the debugger.
+		//
+		// Unlike the other wrapper types, we need these for FUsdStage because it is templated, and we can't express templated
+		// type cast on the .natvis files very well.
+		class UNREALUSDWRAPPER_API FSdfLayerDebugImpl
+		{
+			FString DebugIdentifier;
+			TUsdStore<int> FakePointer;
+		};
+#endif	  // ENABLE_USD_DEBUG_PATH
+
 		template<typename PtrType>
 		class FSdfLayerImpl
 		{
 		public:
 			FSdfLayerImpl() = default;
 
+#if ENABLE_USD_DEBUG_PATH
+			FString DebugIdentifier;
+#endif
+
 			explicit FSdfLayerImpl(const PtrType& InUsdPtr)
 				: UsdPtr(InUsdPtr)
 			{
+#if ENABLE_USD_DEBUG_PATH
+				if (const PtrType& Ptr = GetInner())
+				{
+					DebugIdentifier = FString{UTF8_TO_TCHAR(Ptr->GetIdentifier().c_str())};
+				}
+#endif
 			}
 
 			explicit FSdfLayerImpl(PtrType&& InUsdPtr)
 				: UsdPtr(MoveTemp(InUsdPtr))
 			{
+#if ENABLE_USD_DEBUG_PATH
+				if (const PtrType& Ptr = GetInner())
+				{
+					DebugIdentifier = FString{UTF8_TO_TCHAR(Ptr->GetIdentifier().c_str())};
+				}
+#endif
 			}
 
 			PtrType& GetInner()
@@ -47,7 +76,7 @@ namespace UE
 
 			TUsdStore<PtrType> UsdPtr;
 		};
-	}
+	}	 // namespace Internal
 
 	template<typename PtrType>
 	FSdfLayerBase<PtrType>::FSdfLayerBase()
diff --git a/Engine/Plugins/Runtime/USDCore/Source/UnrealUSDWrapper/Private/UsdWrappers/SdfPayloadsProxy.cpp b/Engine/Plugins/Runtime/USDCore/Source/UnrealUSDWrapper/Private/UsdWrappers/SdfPayloadsProxy.cpp
new file mode 100644
--- /dev/null
+++ b/Engine/Plugins/Runtime/USDCore/Source/UnrealUSDWrapper/Private/UsdWrappers/SdfPayloadsProxy.cpp
@@ -0,0 +1,150 @@
+// Copyright Epic Games, Inc. All Rights Reserved.
+
+#include "UsdWrappers/SdfPayloadsProxy.h"
+
+#include "USDMemory.h"
+
+#include "UsdWrappers/SdfAttributeSpec.h"
+#include "UsdWrappers/SdfLayer.h"
+#include "UsdWrappers/SdfPath.h"
+
+#if USE_USD_SDK
+#include "USDIncludesStart.h"
+#include "pxr/usd/sdf/listEditorProxy.h"
+#include "pxr/usd/sdf/proxyPolicies.h"
+#include "pxr/usd/sdf/payload.h"
+#include "USDIncludesEnd.h"
+#endif	  // #if USE_USD_SDK
+
+namespace UE
+{
+	namespace Internal
+	{
+		class FSdfPayloadsProxyImpl
+		{
+		public:
+			FSdfPayloadsProxyImpl() = default;
+
+#if USE_USD_SDK
+			explicit FSdfPayloadsProxyImpl(const pxr::SdfPayloadsProxy& InSdfPayloadsProxy)
+				: PxrSdfPayloadsProxy(InSdfPayloadsProxy)
+			{
+			}
+
+			explicit FSdfPayloadsProxyImpl(pxr::SdfPayloadsProxy&& InSdfPayloadsProxy)
+				: PxrSdfPayloadsProxy(MoveTemp(InSdfPayloadsProxy))
+			{
+			}
+
+			TUsdStore<pxr::SdfPayloadsProxy> PxrSdfPayloadsProxy;
+#endif	  // #if USE_USD_SDK
+		};
+	}
+
+	FSdfPayloadsProxy::FSdfPayloadsProxy()
+	{
+		FScopedUnrealAllocs UnrealAllocs;
+		Impl = MakeUnique<Internal::FSdfPayloadsProxyImpl>();
+	}
+
+	FSdfPayloadsProxy::FSdfPayloadsProxy(const FSdfPayloadsProxy& Other)
+	{
+#if USE_USD_SDK
+		FScopedUnrealAllocs UnrealAllocs;
+		Impl = MakeUnique<Internal::FSdfPayloadsProxyImpl>(Other.Impl->PxrSdfPayloadsProxy.Get());
+#endif	  // #if USE_USD_SDK
+	}
+
+	FSdfPayloadsProxy::FSdfPayloadsProxy(FSdfPayloadsProxy&& Other) = default;
+
+	FSdfPayloadsProxy::~FSdfPayloadsProxy()
+	{
+		FScopedUnrealAllocs UnrealAllocs;
+		Impl.Reset();
+	}
+
+	FSdfPayloadsProxy& FSdfPayloadsProxy::operator=(const FSdfPayloadsProxy& Other)
+	{
+#if USE_USD_SDK
+		FScopedUnrealAllocs UnrealAllocs;
+		Impl = MakeUnique<Internal::FSdfPayloadsProxyImpl>(Other.Impl->PxrSdfPayloadsProxy.Get());
+#endif	  // #if USE_USD_SDK
+		return *this;
+	}
+
+	FSdfPayloadsProxy& FSdfPayloadsProxy::operator=(FSdfPayloadsProxy&& Other)
+	{
+		FScopedUnrealAllocs UnrealAllocs;
+		Impl = MoveTemp(Other.Impl);
+
+		return *this;
+	}
+
+	FSdfPayloadsProxy::operator bool() const
+	{
+#if USE_USD_SDK
+		return (bool)Impl->PxrSdfPayloadsProxy.Get();
+#else
+		return false;
+#endif	  // #if USE_USD_SDK
+	}
+
+#if USE_USD_SDK
+	FSdfPayloadsProxy::FSdfPayloadsProxy(const pxr::SdfPayloadsProxy& InSdfPayloadsProxy)
+	{
+		FScopedUnrealAllocs UnrealAllocs;
+		Impl = MakeUnique<Internal::FSdfPayloadsProxyImpl>(InSdfPayloadsProxy);
+	}
+
+	FSdfPayloadsProxy::FSdfPayloadsProxy(pxr::SdfPayloadsProxy&& InSdfPayloadsProxy)
+	{
+		FScopedUnrealAllocs UnrealAllocs;
+		Impl = MakeUnique<Internal::FSdfPayloadsProxyImpl>(MoveTemp(InSdfPayloadsProxy));
+	}
+
+	FSdfPayloadsProxy& FSdfPayloadsProxy::operator=(const pxr::SdfPayloadsProxy& InSdfPayloadsProxy)
+	{
+		FScopedUnrealAllocs UnrealAllocs;
+		Impl = MakeUnique<Internal::FSdfPayloadsProxyImpl>(InSdfPayloadsProxy);
+		return *this;
+	}
+
+	FSdfPayloadsProxy& FSdfPayloadsProxy::operator=(pxr::SdfPayloadsProxy&& InSdfPayloadsProxy)
+	{
+		FScopedUnrealAllocs UnrealAllocs;
+		Impl = MakeUnique<Internal::FSdfPayloadsProxyImpl>(MoveTemp(InSdfPayloadsProxy));
+		return *this;
+	}
+
+	FSdfPayloadsProxy::operator pxr::SdfPayloadsProxy&()
+	{
+		return Impl->PxrSdfPayloadsProxy.Get();
+	}
+
+	FSdfPayloadsProxy::operator const pxr::SdfPayloadsProxy&() const
+	{
+		return Impl->PxrSdfPayloadsProxy.Get();
+	}
+#endif	  // #if USE_USD_SDK
+
+	TArray<FSdfPayload> FSdfPayloadsProxy::GetAppliedItems() const
+	{
+		TArray<FSdfPayload> Result;
+#if USE_USD_SDK
+		TUsdStore<std::vector<pxr::SdfPayload>> SdfPayloads = Impl->PxrSdfPayloadsProxy->GetAppliedItems();
+		Result.Reserve(SdfPayloads->size());
+
+		for (const pxr::SdfPayload& SdfPayload : *SdfPayloads)
+		{
+			const pxr::SdfLayerOffset& SdfOffset = SdfPayload.GetLayerOffset();
+
+			Result.Add(FSdfPayload{
+				FString{UTF8_TO_TCHAR(SdfPayload.GetAssetPath().c_str())},
+				FSdfPath{SdfPayload.GetPrimPath()},
+				FSdfLayerOffset{SdfOffset.GetOffset(), SdfOffset.GetScale()}
+			});
+		}
+#endif	  // USE_USD_SDK
+		return Result;
+	}
+}	 // namespace UE
diff --git a/Engine/Plugins/Runtime/USDCore/Source/UnrealUSDWrapper/Private/UsdWrappers/SdfPrimSpec.cpp b/Engine/Plugins/Runtime/USDCore/Source/UnrealUSDWrapper/Private/UsdWrappers/SdfPrimSpec.cpp
--- a/Engine/Plugins/Runtime/USDCore/Source/UnrealUSDWrapper/Private/UsdWrappers/SdfPrimSpec.cpp
+++ b/Engine/Plugins/Runtime/USDCore/Source/UnrealUSDWrapper/Private/UsdWrappers/SdfPrimSpec.cpp
@@ -7,11 +7,14 @@
 #include "UsdWrappers/SdfAttributeSpec.h"
 #include "UsdWrappers/SdfLayer.h"
 #include "UsdWrappers/SdfPath.h"
+#include "UsdWrappers/SdfPayloadsProxy.h"
+#include "UsdWrappers/SdfReferencesProxy.h"
 
 #if USE_USD_SDK
 #include "USDIncludesStart.h"
 #include "pxr/usd/sdf/attributeSpec.h"
 #include "pxr/usd/sdf/primSpec.h"
+#include "pxr/usd/sdf/layer.h"
 #include "USDIncludesEnd.h"
 #endif	  // #if USE_USD_SDK
 
@@ -25,15 +28,42 @@ namespace UE
 			FSdfPrimSpecImpl() = default;
 
 #if USE_USD_SDK
+
+#if ENABLE_USD_DEBUG_PATH
+			FString DebugPath;
+			FString DebugLayerPath;
+#endif
+
 			explicit FSdfPrimSpecImpl(const pxr::SdfPrimSpecHandle& InSdfPrimSpec)
 				: PxrSdfPrimSpec(InSdfPrimSpec)
 			{
+#if ENABLE_USD_DEBUG_PATH
+				SetupDebugMembers();
+#endif
 			}
 
 			explicit FSdfPrimSpecImpl(pxr::SdfPrimSpecHandle&& InSdfPrimSpec)
 				: PxrSdfPrimSpec(MoveTemp(InSdfPrimSpec))
 			{
+#if ENABLE_USD_DEBUG_PATH
+				SetupDebugMembers();
+#endif
+			}
+
+#if ENABLE_USD_DEBUG_PATH
+			void SetupDebugMembers()
+			{
+				if (pxr::SdfPrimSpecHandle ValidSpec = PxrSdfPrimSpec.Get())
+				{
+					DebugPath = FString{UTF8_TO_TCHAR(ValidSpec->GetPath().GetString().c_str())};
+
+					if (pxr::SdfLayerHandle ValidLayer = ValidSpec->GetLayer())
+					{
+						DebugLayerPath = FString{UTF8_TO_TCHAR(ValidLayer->GetIdentifier().c_str())};
+					}
+				}
 			}
+#endif
 
 			TUsdStore<pxr::SdfPrimSpecHandle> PxrSdfPrimSpec;
 #endif	  // #if USE_USD_SDK
@@ -209,6 +239,24 @@ namespace UE
 #endif	  // #if USE_USD_SDK
 	}
 
+	FSdfPayloadsProxy FSdfPrimSpec::GetPayloadList() const
+	{
+#if USE_USD_SDK
+		return UE::FSdfPayloadsProxy{Impl->PxrSdfPrimSpec.Get()->GetPayloadList()};
+#else
+		return UE::FSdfPayloadsProxy{};
+#endif	  // #if USE_USD_SDK
+	}
+
+	FSdfReferencesProxy FSdfPrimSpec::GetReferenceList() const
+	{
+#if USE_USD_SDK
+		return UE::FSdfReferencesProxy{Impl->PxrSdfPrimSpec.Get()->GetReferenceList()};
+#else
+		return UE::FSdfReferencesProxy{};
+#endif	  // #if USE_USD_SDK
+	}
+
 	FName FSdfPrimSpec::GetTypeName() const
 	{
 #if USE_USD_SDK
diff --git a/Engine/Plugins/Runtime/USDCore/Source/UnrealUSDWrapper/Private/UsdWrappers/SdfReferencesProxy.cpp b/Engine/Plugins/Runtime/USDCore/Source/UnrealUSDWrapper/Private/UsdWrappers/SdfReferencesProxy.cpp
new file mode 100644
--- /dev/null
+++ b/Engine/Plugins/Runtime/USDCore/Source/UnrealUSDWrapper/Private/UsdWrappers/SdfReferencesProxy.cpp
@@ -0,0 +1,150 @@
+// Copyright Epic Games, Inc. All Rights Reserved.
+
+#include "UsdWrappers/SdfReferencesProxy.h"
+
+#include "USDMemory.h"
+
+#include "UsdWrappers/SdfAttributeSpec.h"
+#include "UsdWrappers/SdfLayer.h"
+#include "UsdWrappers/SdfPath.h"
+
+#if USE_USD_SDK
+#include "USDIncludesStart.h"
+#include "pxr/usd/sdf/listEditorProxy.h"
+#include "pxr/usd/sdf/proxyPolicies.h"
+#include "pxr/usd/sdf/reference.h"
+#include "USDIncludesEnd.h"
+#endif	  // #if USE_USD_SDK
+
+namespace UE
+{
+	namespace Internal
+	{
+		class FSdfReferencesProxyImpl
+		{
+		public:
+			FSdfReferencesProxyImpl() = default;
+
+#if USE_USD_SDK
+			explicit FSdfReferencesProxyImpl(const pxr::SdfReferencesProxy& InSdfReferencesProxy)
+				: PxrSdfReferencesProxy(InSdfReferencesProxy)
+			{
+			}
+
+			explicit FSdfReferencesProxyImpl(pxr::SdfReferencesProxy&& InSdfReferencesProxy)
+				: PxrSdfReferencesProxy(MoveTemp(InSdfReferencesProxy))
+			{
+			}
+
+			TUsdStore<pxr::SdfReferencesProxy> PxrSdfReferencesProxy;
+#endif	  // #if USE_USD_SDK
+		};
+	}
+
+	FSdfReferencesProxy::FSdfReferencesProxy()
+	{
+		FScopedUnrealAllocs UnrealAllocs;
+		Impl = MakeUnique<Internal::FSdfReferencesProxyImpl>();
+	}
+
+	FSdfReferencesProxy::FSdfReferencesProxy(const FSdfReferencesProxy& Other)
+	{
+#if USE_USD_SDK
+		FScopedUnrealAllocs UnrealAllocs;
+		Impl = MakeUnique<Internal::FSdfReferencesProxyImpl>(Other.Impl->PxrSdfReferencesProxy.Get());
+#endif	  // #if USE_USD_SDK
+	}
+
+	FSdfReferencesProxy::FSdfReferencesProxy(FSdfReferencesProxy&& Other) = default;
+
+	FSdfReferencesProxy::~FSdfReferencesProxy()
+	{
+		FScopedUnrealAllocs UnrealAllocs;
+		Impl.Reset();
+	}
+
+	FSdfReferencesProxy& FSdfReferencesProxy::operator=(const FSdfReferencesProxy& Other)
+	{
+#if USE_USD_SDK
+		FScopedUnrealAllocs UnrealAllocs;
+		Impl = MakeUnique<Internal::FSdfReferencesProxyImpl>(Other.Impl->PxrSdfReferencesProxy.Get());
+#endif	  // #if USE_USD_SDK
+		return *this;
+	}
+
+	FSdfReferencesProxy& FSdfReferencesProxy::operator=(FSdfReferencesProxy&& Other)
+	{
+		FScopedUnrealAllocs UnrealAllocs;
+		Impl = MoveTemp(Other.Impl);
+
+		return *this;
+	}
+
+	FSdfReferencesProxy::operator bool() const
+	{
+#if USE_USD_SDK
+		return (bool)Impl->PxrSdfReferencesProxy.Get();
+#else
+		return false;
+#endif	  // #if USE_USD_SDK
+	}
+
+#if USE_USD_SDK
+	FSdfReferencesProxy::FSdfReferencesProxy(const pxr::SdfReferencesProxy& InSdfReferencesProxy)
+	{
+		FScopedUnrealAllocs UnrealAllocs;
+		Impl = MakeUnique<Internal::FSdfReferencesProxyImpl>(InSdfReferencesProxy);
+	}
+
+	FSdfReferencesProxy::FSdfReferencesProxy(pxr::SdfReferencesProxy&& InSdfReferencesProxy)
+	{
+		FScopedUnrealAllocs UnrealAllocs;
+		Impl = MakeUnique<Internal::FSdfReferencesProxyImpl>(MoveTemp(InSdfReferencesProxy));
+	}
+
+	FSdfReferencesProxy& FSdfReferencesProxy::operator=(const pxr::SdfReferencesProxy& InSdfReferencesProxy)
+	{
+		FScopedUnrealAllocs UnrealAllocs;
+		Impl = MakeUnique<Internal::FSdfReferencesProxyImpl>(InSdfReferencesProxy);
+		return *this;
+	}
+
+	FSdfReferencesProxy& FSdfReferencesProxy::operator=(pxr::SdfReferencesProxy&& InSdfReferencesProxy)
+	{
+		FScopedUnrealAllocs UnrealAllocs;
+		Impl = MakeUnique<Internal::FSdfReferencesProxyImpl>(MoveTemp(InSdfReferencesProxy));
+		return *this;
+	}
+
+	FSdfReferencesProxy::operator pxr::SdfReferencesProxy&()
+	{
+		return Impl->PxrSdfReferencesProxy.Get();
+	}
+
+	FSdfReferencesProxy::operator const pxr::SdfReferencesProxy&() const
+	{
+		return Impl->PxrSdfReferencesProxy.Get();
+	}
+#endif	  // #if USE_USD_SDK
+
+	TArray<FSdfReference> FSdfReferencesProxy::GetAppliedItems() const
+	{
+		TArray<FSdfReference> Result;
+#if USE_USD_SDK
+		TUsdStore<std::vector<pxr::SdfReference>> SdfReferences = Impl->PxrSdfReferencesProxy->GetAppliedItems();
+		Result.Reserve(SdfReferences->size());
+
+		for (const pxr::SdfReference& SdfReference : *SdfReferences)
+		{
+			const pxr::SdfLayerOffset& SdfOffset = SdfReference.GetLayerOffset();
+
+			Result.Add(FSdfReference{
+				FString{UTF8_TO_TCHAR(SdfReference.GetAssetPath().c_str())},
+				FSdfPath{SdfReference.GetPrimPath()},
+				FSdfLayerOffset{SdfOffset.GetOffset(), SdfOffset.GetScale()}
+			});
+		}
+#endif	  // USE_USD_SDK
+		return Result;
+	}
+}	 // namespace UE
diff --git a/Engine/Plugins/Runtime/USDCore/Source/UnrealUSDWrapper/Private/UsdWrappers/UsdAttribute.cpp b/Engine/Plugins/Runtime/USDCore/Source/UnrealUSDWrapper/Private/UsdWrappers/UsdAttribute.cpp
--- a/Engine/Plugins/Runtime/USDCore/Source/UnrealUSDWrapper/Private/UsdWrappers/UsdAttribute.cpp
+++ b/Engine/Plugins/Runtime/USDCore/Source/UnrealUSDWrapper/Private/UsdWrappers/UsdAttribute.cpp
@@ -4,6 +4,7 @@
 
 #include "UsdWrappers/SdfPath.h"
 #include "UsdWrappers/UsdPrim.h"
+#include "UsdWrappers/UsdResolveInfo.h"
 #include "UsdWrappers/VtValue.h"
 
 #include "USDMemory.h"
@@ -24,24 +25,32 @@ namespace UE
 			FUsdAttributeImpl() = default;
 
 #if USE_USD_SDK
+
+#if ENABLE_USD_DEBUG_PATH
+			FString DebugPath;
+#endif
+
 			explicit FUsdAttributeImpl(const pxr::UsdAttribute& InUsdAttribute)
 				: PxrUsdAttribute(InUsdAttribute)
 			{
+#if ENABLE_USD_DEBUG_PATH
+				DebugPath = FString{UTF8_TO_TCHAR(PxrUsdAttribute.Get().GetPath().GetString().c_str())};
+#endif
 			}
 
 			explicit FUsdAttributeImpl(pxr::UsdAttribute&& InUsdAttribute)
 				: PxrUsdAttribute(MoveTemp(InUsdAttribute))
 			{
+#if ENABLE_USD_DEBUG_PATH
+				DebugPath = FString{UTF8_TO_TCHAR(PxrUsdAttribute.Get().GetPath().GetString().c_str())};
+#endif
 			}
 
 			TUsdStore<pxr::UsdAttribute> PxrUsdAttribute;
 #endif	  // #if USE_USD_SDK
 		};
-	}
-}
+	}	 // namespace Internal
 
-namespace UE
-{
 	FUsdAttribute::FUsdAttribute()
 	{
 		FScopedUnrealAllocs UnrealAllocs;
@@ -302,6 +311,16 @@ namespace UE
 #endif	  // #if USE_USD_SDK
 	}
 
+	FUsdResolveInfo FUsdAttribute::GetResolveInfo(double Time) const
+	{
+#if USE_USD_SDK
+		FScopedUsdAllocs UsdAllocs;
+		return FUsdResolveInfo{Impl->PxrUsdAttribute.Get().GetResolveInfo(Time)};
+#else
+		return {};
+#endif	  // #if USE_USD_SDK
+	}
+
 	bool FUsdAttribute::Set(const UE::FVtValue& Value, TOptional<double> Time /*= {} */) const
 	{
 #if USE_USD_SDK
diff --git a/Engine/Plugins/Runtime/USDCore/Source/UnrealUSDWrapper/Private/UsdWrappers/UsdPrim.cpp b/Engine/Plugins/Runtime/USDCore/Source/UnrealUSDWrapper/Private/UsdWrappers/UsdPrim.cpp
--- a/Engine/Plugins/Runtime/USDCore/Source/UnrealUSDWrapper/Private/UsdWrappers/UsdPrim.cpp
+++ b/Engine/Plugins/Runtime/USDCore/Source/UnrealUSDWrapper/Private/UsdWrappers/UsdPrim.cpp
@@ -4,7 +4,9 @@
 
 #include "USDMemory.h"
 
+#include "UsdWrappers/PcpPrimIndex.h"
 #include "UsdWrappers/SdfPath.h"
+#include "UsdWrappers/SdfPrimSpec.h"
 #include "UsdWrappers/UsdAttribute.h"
 #include "UsdWrappers/UsdPayloads.h"
 #include "UsdWrappers/UsdReferences.h"
@@ -34,13 +36,13 @@ namespace UE
 #if USE_USD_SDK
 
 #if ENABLE_USD_DEBUG_PATH
-			FSdfPath DebugPath;
+			FString DebugPath;
 #endif
 			explicit FUsdPrimImpl(const pxr::UsdPrim& InUsdPrim)
 				: PxrUsdPrim(InUsdPrim)
 			{
 #if ENABLE_USD_DEBUG_PATH
-				DebugPath = PxrUsdPrim.Get().GetPrimPath();
+				DebugPath = FString{UTF8_TO_TCHAR(PxrUsdPrim.Get().GetPrimPath().GetString().c_str())};
 #endif
 			}
 
@@ -48,7 +50,7 @@ namespace UE
 				: PxrUsdPrim(MoveTemp(InUsdPrim))
 			{
 #if ENABLE_USD_DEBUG_PATH
-				DebugPath = PxrUsdPrim.Get().GetPrimPath();
+				DebugPath = FString{UTF8_TO_TCHAR(PxrUsdPrim.Get().GetPrimPath().GetString().c_str())};
 #endif
 			}
 
@@ -671,6 +673,41 @@ namespace UE
 #endif	  // #if USE_USD_SDK
 	}
 
+	TArray<FSdfPrimSpec> FUsdPrim::GetPrimStack() const
+	{
+		TArray<FSdfPrimSpec> Result;
+
+#if USE_USD_SDK
+		TUsdStore<std::vector<pxr::SdfPrimSpecHandle>> UsdSpecStack = Impl->PxrUsdPrim.Get().GetPrimStack();
+
+		Result.Reserve(UsdSpecStack->size());
+		for (const pxr::SdfPrimSpecHandle& Handle : *UsdSpecStack)
+		{
+			Result.Add(FSdfPrimSpec{Handle});
+		}
+#endif	  // USE_USD_SDK
+
+		return Result;
+	}
+
+	FPcpPrimIndex FUsdPrim::GetPrimIndex() const
+	{
+#if USE_USD_SDK
+		return FPcpPrimIndex{Impl->PxrUsdPrim.Get().GetPrimIndex()};
+#else
+		return {};
+#endif	  // #if USE_USD_SDK
+	}
+
+	FPcpPrimIndex FUsdPrim::ComputeExpandedPrimIndex() const
+	{
+#if USE_USD_SDK
+		return FPcpPrimIndex{Impl->PxrUsdPrim.Get().ComputeExpandedPrimIndex()};
+#else
+		return {};
+#endif	  // #if USE_USD_SDK
+	}
+
 	FUsdAttribute FUsdPrim::GetAttribute(const TCHAR* AttrName) const
 	{
 #if USE_USD_SDK
diff --git a/Engine/Plugins/Runtime/USDCore/Source/UnrealUSDWrapper/Private/UsdWrappers/UsdResolveInfo.cpp b/Engine/Plugins/Runtime/USDCore/Source/UnrealUSDWrapper/Private/UsdWrappers/UsdResolveInfo.cpp
new file mode 100644
--- /dev/null
+++ b/Engine/Plugins/Runtime/USDCore/Source/UnrealUSDWrapper/Private/UsdWrappers/UsdResolveInfo.cpp
@@ -0,0 +1,97 @@
+// Copyright Epic Games, Inc. All Rights Reserved.
+
+#include "UsdWrappers/UsdResolveInfo.h"
+
+#include "USDMemory.h"
+#include "UsdWrappers/PcpNodeRef.h"
+
+#if USE_USD_SDK
+#include "USDIncludesStart.h"
+#include "pxr/usd/usd/resolveInfo.h"
+#include "USDIncludesEnd.h"
+#endif	  // #if USE_USD_SDK
+
+namespace UE
+{
+	namespace Internal
+	{
+		class FUsdResolveInfoImpl
+		{
+		public:
+			FUsdResolveInfoImpl() = default;
+
+#if USE_USD_SDK
+			explicit FUsdResolveInfoImpl(const pxr::UsdResolveInfo& InUsdResolveInfo)
+				: PxrUsdResolveInfo(InUsdResolveInfo)
+			{
+			}
+
+			explicit FUsdResolveInfoImpl(pxr::UsdResolveInfo&& InUsdResolveInfo)
+				: PxrUsdResolveInfo(MoveTemp(InUsdResolveInfo))
+			{
+			}
+
+			TUsdStore<pxr::UsdResolveInfo> PxrUsdResolveInfo;
+#endif	  // #if USE_USD_SDK
+		};
+	}
+}
+
+namespace UE
+{
+	FUsdResolveInfo::FUsdResolveInfo()
+	{
+		FScopedUnrealAllocs UnrealAllocs;
+		Impl = MakeUnique<Internal::FUsdResolveInfoImpl>();
+	}
+
+	FUsdResolveInfo::FUsdResolveInfo(const FUsdResolveInfo& Other)
+	{
+#if USE_USD_SDK
+		FScopedUnrealAllocs UnrealAllocs;
+		Impl = MakeUnique<Internal::FUsdResolveInfoImpl>(Other.Impl->PxrUsdResolveInfo.Get());
+#endif	  // #if USE_USD_SDK
+	}
+
+	FUsdResolveInfo::FUsdResolveInfo(FUsdResolveInfo&& Other) = default;
+
+	FUsdResolveInfo::~FUsdResolveInfo()
+	{
+		FScopedUnrealAllocs UnrealAllocs;
+		Impl.Reset();
+	}
+
+#if USE_USD_SDK
+	FUsdResolveInfo::FUsdResolveInfo(const pxr::UsdResolveInfo& InUsdResolveInfo)
+	{
+		FScopedUnrealAllocs UnrealAllocs;
+		Impl = MakeUnique<Internal::FUsdResolveInfoImpl>(InUsdResolveInfo);
+	}
+
+	FUsdResolveInfo::FUsdResolveInfo(pxr::UsdResolveInfo&& InUsdResolveInfo)
+	{
+		FScopedUnrealAllocs UnrealAllocs;
+		Impl = MakeUnique<Internal::FUsdResolveInfoImpl>(MoveTemp(InUsdResolveInfo));
+	}
+
+	FUsdResolveInfo::operator pxr::UsdResolveInfo&()
+	{
+		return Impl->PxrUsdResolveInfo.Get();
+	}
+
+	FUsdResolveInfo::operator const pxr::UsdResolveInfo&() const
+	{
+		return Impl->PxrUsdResolveInfo.Get();
+	}
+#endif	  // #if USE_USD_SDK
+
+	UE::FPcpNodeRef FUsdResolveInfo::GetNode() const
+	{
+#if USE_USD_SDK
+		FScopedUsdAllocs UsdAllocs;
+		return UE::FPcpNodeRef{Impl->PxrUsdResolveInfo.Get().GetNode()};
+#else
+		return {};
+#endif	  // #if USE_USD_SDK
+	}
+}	 // namespace UE
diff --git a/Engine/Plugins/Runtime/USDCore/Source/UnrealUSDWrapper/Private/UsdWrappers/UsdStage.cpp b/Engine/Plugins/Runtime/USDCore/Source/UnrealUSDWrapper/Private/UsdWrappers/UsdStage.cpp
--- a/Engine/Plugins/Runtime/USDCore/Source/UnrealUSDWrapper/Private/UsdWrappers/UsdStage.cpp
+++ b/Engine/Plugins/Runtime/USDCore/Source/UnrealUSDWrapper/Private/UsdWrappers/UsdStage.cpp
@@ -21,20 +21,53 @@ namespace UE
 {
 	namespace Internal
 	{
+#if ENABLE_USD_DEBUG_PATH
+		// The only purpose of this class is so that we can cast the UsdStage impl to this type in the Unreal.natvis debug
+		// visualizers, allowing us to inspect these FUsdStage and FUsdStageWeak variables in the debugger.
+		//
+		// Unlike the other wrapper types, we need these for FUsdStage because it is templated, and we can't express templated
+		// type cast on the .natvis files very well.
+		class UNREALUSDWRAPPER_API FUsdStageDebugImpl
+		{
+			FString DebugRootLayerIdentifier;
+			FString DebugSessionLayerIdentifier;
+			TUsdStore<int> FakePointer;
+		};
+#endif	  // ENABLE_USD_DEBUG_PATH
+
 		template<typename PtrType>
 		class FUsdStageImpl
 		{
 		public:
 			FUsdStageImpl() = default;
 
+#if ENABLE_USD_DEBUG_PATH
+			FString DebugRootLayerIdentifier;
+			FString DebugSessionLayerIdentifier;
+#endif
+
 			explicit FUsdStageImpl(const PtrType& InUsdPtr)
 				: UsdPtr(InUsdPtr)
 			{
+#if ENABLE_USD_DEBUG_PATH
+				if (const PtrType& Ptr = GetInner())
+				{
+					DebugRootLayerIdentifier = FString{UTF8_TO_TCHAR(Ptr->GetRootLayer()->GetIdentifier().c_str())};
+					DebugSessionLayerIdentifier = FString{UTF8_TO_TCHAR(Ptr->GetSessionLayer()->GetIdentifier().c_str())};
+				}
+#endif
 			}
 
 			explicit FUsdStageImpl(PtrType&& InUsdPtr)
 				: UsdPtr(MoveTemp(InUsdPtr))
 			{
+#if ENABLE_USD_DEBUG_PATH
+				if (const PtrType& Ptr = GetInner())
+				{
+					DebugRootLayerIdentifier = FString{UTF8_TO_TCHAR(Ptr->GetRootLayer()->GetIdentifier().c_str())};
+					DebugSessionLayerIdentifier = FString{UTF8_TO_TCHAR(Ptr->GetSessionLayer()->GetIdentifier().c_str())};
+				}
+#endif
 			}
 
 			PtrType& GetInner()
@@ -49,7 +82,7 @@ namespace UE
 
 			TUsdStore<PtrType> UsdPtr;
 		};
-	}
+	}	 // namespace Internal
 
 	template<typename PtrType>
 	FUsdStageBase<PtrType>::FUsdStageBase()
diff --git a/Engine/Plugins/Runtime/USDCore/Source/UnrealUSDWrapper/Public/USDMemory.h b/Engine/Plugins/Runtime/USDCore/Source/UnrealUSDWrapper/Public/USDMemory.h
--- a/Engine/Plugins/Runtime/USDCore/Source/UnrealUSDWrapper/Public/USDMemory.h
+++ b/Engine/Plugins/Runtime/USDCore/Source/UnrealUSDWrapper/Public/USDMemory.h
@@ -272,6 +272,11 @@ public:
 		return *this;
 	}
 
+	explicit operator bool() const
+	{
+		return StoredUsdObject.IsSet();
+	}
+
 	~TUsdStore()
 	{
 		FScopedUsdAllocs UsdAllocs;
diff --git a/Engine/Plugins/Runtime/USDCore/Source/UnrealUSDWrapper/Public/UsdWrappers/PcpLayerStack.h b/Engine/Plugins/Runtime/USDCore/Source/UnrealUSDWrapper/Public/UsdWrappers/PcpLayerStack.h
new file mode 100644
--- /dev/null
+++ b/Engine/Plugins/Runtime/USDCore/Source/UnrealUSDWrapper/Public/UsdWrappers/PcpLayerStack.h
@@ -0,0 +1,78 @@
+// Copyright Epic Games, Inc. All Rights Reserved.
+
+#pragma once
+
+#include "Templates/UniquePtr.h"
+#include "UsdWrappers/ForwardDeclarations.h"
+
+#if USE_USD_SDK
+
+#include "USDIncludesStart.h"
+#include "pxr/pxr.h"
+#include "USDIncludesEnd.h"
+
+PXR_NAMESPACE_OPEN_SCOPE
+	class PcpLayerStack;
+
+	template<typename T>
+	class TfRefPtr;
+
+	using PcpLayerStackRefPtr = TfRefPtr<PcpLayerStack>;
+PXR_NAMESPACE_CLOSE_SCOPE
+
+#endif	  // #if USE_USD_SDK
+
+namespace UE
+{
+	class FPcpLayerStackIdentifier;
+	struct FSdfLayerOffset;
+	namespace Internal
+	{
+		class FPcpLayerStackImpl;
+	}
+
+	/**
+	 * Minimal pxr::PcpLayerStackRefPtr wrapper for Unreal that can be used from no-rtti modules.
+	 */
+	class UNREALUSDWRAPPER_API FPcpLayerStack
+	{
+	public:
+		FPcpLayerStack();
+
+		FPcpLayerStack(const FPcpLayerStack& Other);
+		FPcpLayerStack(FPcpLayerStack&& Other);
+		~FPcpLayerStack();
+
+		FPcpLayerStack& operator=(const FPcpLayerStack& Other);
+		FPcpLayerStack& operator=(FPcpLayerStack&& Other);
+
+		bool operator==(const FPcpLayerStack& Other) const;
+		bool operator!=(const FPcpLayerStack& Other) const;
+
+		explicit operator bool() const;
+
+		// Auto conversion from/to pxr::PcpLayerStackRefPtr
+	public:
+#if USE_USD_SDK
+		explicit FPcpLayerStack(const pxr::PcpLayerStackRefPtr& InPcpLayerStack);
+		explicit FPcpLayerStack(pxr::PcpLayerStackRefPtr&& InPcpLayerStack);
+		FPcpLayerStack& operator=(const pxr::PcpLayerStackRefPtr& InPcpLayerStack);
+		FPcpLayerStack& operator=(pxr::PcpLayerStackRefPtr&& InPcpLayerStack);
+
+		operator pxr::PcpLayerStackRefPtr&();
+		operator const pxr::PcpLayerStackRefPtr&() const;
+#endif	  // #if USE_USD_SDK
+
+		  // Wrapped pxr::PcpLayerStack functions, refer to the USD SDK documentation
+	public:
+		FPcpLayerStackIdentifier GetIdentifier() const;
+		TArray<FSdfLayer> GetLayers() const;
+		FSdfLayerOffset GetLayerOffsetForLayer(const FSdfLayer& Layer) const;
+		FSdfLayerOffset GetLayerOffsetForLayer(const FSdfLayerWeak& Layer) const;
+		bool HasLayer(const FSdfLayer& Layer) const;
+		bool HasLayer(const FSdfLayerWeak& Layer) const;
+
+	private:
+		TUniquePtr<Internal::FPcpLayerStackImpl> Impl;
+	};
+}	 // namespace UE
diff --git a/Engine/Plugins/Runtime/USDCore/Source/UnrealUSDWrapper/Public/UsdWrappers/PcpLayerStackIdentifier.h b/Engine/Plugins/Runtime/USDCore/Source/UnrealUSDWrapper/Public/UsdWrappers/PcpLayerStackIdentifier.h
new file mode 100644
--- /dev/null
+++ b/Engine/Plugins/Runtime/USDCore/Source/UnrealUSDWrapper/Public/UsdWrappers/PcpLayerStackIdentifier.h
@@ -0,0 +1,67 @@
+// Copyright Epic Games, Inc. All Rights Reserved.
+
+#pragma once
+
+#include "Templates/UniquePtr.h"
+#include "UsdWrappers/ForwardDeclarations.h"
+
+#if USE_USD_SDK
+
+#include "USDIncludesStart.h"
+#include "pxr/pxr.h"
+#include "USDIncludesEnd.h"
+
+PXR_NAMESPACE_OPEN_SCOPE
+	class PcpLayerStackIdentifier;
+PXR_NAMESPACE_CLOSE_SCOPE
+
+#endif	  // #if USE_USD_SDK
+
+namespace UE
+{
+	namespace Internal
+	{
+		class FPcpLayerStackIdentifierImpl;
+	}
+
+	/**
+	 * Minimal pxr::PcpLayerStackIdentifier wrapper for Unreal that can be used from no-rtti modules.
+	 */
+	class UNREALUSDWRAPPER_API FPcpLayerStackIdentifier
+	{
+	public:
+		FPcpLayerStackIdentifier();
+
+		FPcpLayerStackIdentifier(const FPcpLayerStackIdentifier& Other);
+		FPcpLayerStackIdentifier(FPcpLayerStackIdentifier&& Other);
+		~FPcpLayerStackIdentifier();
+
+		FPcpLayerStackIdentifier& operator=(const FPcpLayerStackIdentifier& Other);
+		FPcpLayerStackIdentifier& operator=(FPcpLayerStackIdentifier&& Other);
+
+		bool operator==(const FPcpLayerStackIdentifier& Other) const;
+		bool operator!=(const FPcpLayerStackIdentifier& Other) const;
+
+		explicit operator bool() const;
+
+		// Auto conversion from/to pxr::PcpLayerStackIdentifier
+	public:
+#if USE_USD_SDK
+		explicit FPcpLayerStackIdentifier(const pxr::PcpLayerStackIdentifier& InPcpLayerStackIdentifier);
+		explicit FPcpLayerStackIdentifier(pxr::PcpLayerStackIdentifier&& InPcpLayerStackIdentifier);
+		FPcpLayerStackIdentifier& operator=(const pxr::PcpLayerStackIdentifier& InPcpLayerStackIdentifier);
+		FPcpLayerStackIdentifier& operator=(pxr::PcpLayerStackIdentifier&& InPcpLayerStackIdentifier);
+
+		operator pxr::PcpLayerStackIdentifier&();
+		operator const pxr::PcpLayerStackIdentifier&() const;
+#endif	  // #if USE_USD_SDK
+
+		  // Wrapped pxr::PcpLayerStackIdentifier functions, refer to the USD SDK documentation
+	public:
+		FSdfLayerWeak RootLayer() const;
+		FSdfLayerWeak SessionLayer() const;
+
+	private:
+		TUniquePtr<Internal::FPcpLayerStackIdentifierImpl> Impl;
+	};
+}	 // namespace UE
diff --git a/Engine/Plugins/Runtime/USDCore/Source/UnrealUSDWrapper/Public/UsdWrappers/PcpMapExpression.h b/Engine/Plugins/Runtime/USDCore/Source/UnrealUSDWrapper/Public/UsdWrappers/PcpMapExpression.h
new file mode 100644
--- /dev/null
+++ b/Engine/Plugins/Runtime/USDCore/Source/UnrealUSDWrapper/Public/UsdWrappers/PcpMapExpression.h
@@ -0,0 +1,61 @@
+// Copyright Epic Games, Inc. All Rights Reserved.
+
+#pragma once
+
+#include "Templates/UniquePtr.h"
+
+#if USE_USD_SDK
+
+#include "USDIncludesStart.h"
+#include "pxr/pxr.h"
+#include "USDIncludesEnd.h"
+
+PXR_NAMESPACE_OPEN_SCOPE
+	class PcpMapExpression;
+PXR_NAMESPACE_CLOSE_SCOPE
+
+#endif	  // #if USE_USD_SDK
+
+namespace UE
+{
+	struct FSdfLayerOffset;
+	namespace Internal
+	{
+		class FPcpMapExpressionImpl;
+	}
+
+	/**
+	 * Minimal pxr::PcpMapExpression wrapper for Unreal that can be used from no-rtti modules.
+	 */
+	class UNREALUSDWRAPPER_API FPcpMapExpression
+	{
+	public:
+		FPcpMapExpression();
+
+		FPcpMapExpression(const FPcpMapExpression& Other);
+		FPcpMapExpression(FPcpMapExpression&& Other);
+		~FPcpMapExpression();
+
+		FPcpMapExpression& operator=(const FPcpMapExpression& Other);
+		FPcpMapExpression& operator=(FPcpMapExpression&& Other);
+
+		// Auto conversion from/to pxr::PcpMapExpression
+	public:
+#if USE_USD_SDK
+		explicit FPcpMapExpression(const pxr::PcpMapExpression& InPcpMapExpression);
+		explicit FPcpMapExpression(pxr::PcpMapExpression&& InPcpMapExpression);
+		FPcpMapExpression& operator=(const pxr::PcpMapExpression& InPcpMapExpression);
+		FPcpMapExpression& operator=(pxr::PcpMapExpression&& InPcpMapExpression);
+
+		operator pxr::PcpMapExpression&();
+		operator const pxr::PcpMapExpression&() const;
+#endif	  // #if USE_USD_SDK
+
+		  // Wrapped pxr::PcpMapExpression functions, refer to the USD SDK documentation
+	public:
+		FSdfLayerOffset GetTimeOffset() const;
+
+	private:
+		TUniquePtr<Internal::FPcpMapExpressionImpl> Impl;
+	};
+}	 // namespace UE
diff --git a/Engine/Plugins/Runtime/USDCore/Source/UnrealUSDWrapper/Public/UsdWrappers/PcpNodeRef.h b/Engine/Plugins/Runtime/USDCore/Source/UnrealUSDWrapper/Public/UsdWrappers/PcpNodeRef.h
new file mode 100644
--- /dev/null
+++ b/Engine/Plugins/Runtime/USDCore/Source/UnrealUSDWrapper/Public/UsdWrappers/PcpNodeRef.h
@@ -0,0 +1,87 @@
+// Copyright Epic Games, Inc. All Rights Reserved.
+
+#pragma once
+
+#include "Templates/UniquePtr.h"
+
+#if USE_USD_SDK
+
+#include "USDIncludesStart.h"
+#include "pxr/pxr.h"
+#include "USDIncludesEnd.h"
+
+PXR_NAMESPACE_OPEN_SCOPE
+	class PcpNodeRef;
+PXR_NAMESPACE_CLOSE_SCOPE
+
+#endif	  // #if USE_USD_SDK
+
+namespace UE
+{
+	class FSdfPath;
+	class FPcpLayerStack;
+	class FPcpMapExpression;
+	namespace Internal
+	{
+		class FPcpNodeRefImpl;
+	}
+
+	/** Analogous to pxr::PcpArcType */
+	enum class EPcpArcType
+	{
+		PcpArcTypeRoot,
+		PcpArcTypeInherit,
+		PcpArcTypeVariant,
+		PcpArcTypeRelocate,
+		PcpArcTypeReference,
+		PcpArcTypePayload,
+		PcpArcTypeSpecialize,
+
+		PcpNumArcTypes
+	};
+
+	/**
+	 * Minimal pxr::PcpNodeRef wrapper for Unreal that can be used from no-rtti modules.
+	 */
+	class UNREALUSDWRAPPER_API FPcpNodeRef
+	{
+	public:
+		FPcpNodeRef();
+
+		FPcpNodeRef(const FPcpNodeRef& Other);
+		FPcpNodeRef(FPcpNodeRef&& Other);
+		~FPcpNodeRef();
+
+		FPcpNodeRef& operator=(const FPcpNodeRef& Other);
+		FPcpNodeRef& operator=(FPcpNodeRef&& Other);
+
+		bool operator==(const FPcpNodeRef& Other) const;
+		bool operator!=(const FPcpNodeRef& Other) const;
+
+		explicit operator bool() const;
+
+		// Auto conversion from/to pxr::PcpNodeRef
+	public:
+#if USE_USD_SDK
+		explicit FPcpNodeRef(const pxr::PcpNodeRef& InPcpNodeRef);
+		explicit FPcpNodeRef(pxr::PcpNodeRef&& InPcpNodeRef);
+		FPcpNodeRef& operator=(const pxr::PcpNodeRef& InPcpNodeRef);
+		FPcpNodeRef& operator=(pxr::PcpNodeRef&& InPcpNodeRef);
+
+		operator pxr::PcpNodeRef&();
+		operator const pxr::PcpNodeRef&() const;
+#endif	  // #if USE_USD_SDK
+
+		  // Wrapped pxr::PcpNodeRef functions, refer to the USD SDK documentation
+	public:
+		EPcpArcType GetArcType() const;
+		FPcpNodeRef GetParentNode() const;
+		FPcpMapExpression GetMapToParent() const;
+		FPcpMapExpression GetMapToRoot() const;
+		FSdfPath GetPath() const;
+		FPcpLayerStack GetLayerStack() const;
+
+	private:
+		TUniquePtr<Internal::FPcpNodeRefImpl> Impl;
+	};
+}	 // namespace UE
diff --git a/Engine/Plugins/Runtime/USDCore/Source/UnrealUSDWrapper/Public/UsdWrappers/PcpPrimIndex.h b/Engine/Plugins/Runtime/USDCore/Source/UnrealUSDWrapper/Public/UsdWrappers/PcpPrimIndex.h
new file mode 100644
--- /dev/null
+++ b/Engine/Plugins/Runtime/USDCore/Source/UnrealUSDWrapper/Public/UsdWrappers/PcpPrimIndex.h
@@ -0,0 +1,65 @@
+// Copyright Epic Games, Inc. All Rights Reserved.
+
+#pragma once
+
+#include "Templates/UniquePtr.h"
+
+#if USE_USD_SDK
+
+#include "USDIncludesStart.h"
+#include "pxr/pxr.h"
+#include "USDIncludesEnd.h"
+
+PXR_NAMESPACE_OPEN_SCOPE
+	class PcpPrimIndex;
+PXR_NAMESPACE_CLOSE_SCOPE
+
+#endif	  // #if USE_USD_SDK
+
+namespace UE
+{
+	class FPcpNodeRef;
+	namespace Internal
+	{
+		class FPcpPrimIndexImpl;
+	}
+
+	/**
+	 * Minimal pxr::PcpPrimIndex wrapper for Unreal that can be used from no-rtti modules.
+	 */
+	class UNREALUSDWRAPPER_API FPcpPrimIndex
+	{
+	public:
+		FPcpPrimIndex();
+
+		FPcpPrimIndex(const FPcpPrimIndex& Other);
+		FPcpPrimIndex(FPcpPrimIndex&& Other);
+		~FPcpPrimIndex();
+
+		FPcpPrimIndex& operator=(const FPcpPrimIndex& Other);
+		FPcpPrimIndex& operator=(FPcpPrimIndex&& Other);
+
+		explicit operator bool() const;
+
+		// Auto conversion from/to pxr::PcpPrimIndex
+	public:
+#if USE_USD_SDK
+		explicit FPcpPrimIndex(const pxr::PcpPrimIndex& InPcpPrimIndex);
+		explicit FPcpPrimIndex(pxr::PcpPrimIndex&& InPcpPrimIndex);
+		FPcpPrimIndex& operator=(const pxr::PcpPrimIndex& InPcpPrimIndex);
+		FPcpPrimIndex& operator=(pxr::PcpPrimIndex&& InPcpPrimIndex);
+
+		operator pxr::PcpPrimIndex&();
+		operator const pxr::PcpPrimIndex&() const;
+#endif	  // #if USE_USD_SDK
+
+		  // Wrapped pxr::PcpPrimIndex functions, refer to the USD SDK documentation
+	public:
+		bool IsValid() const;
+		FPcpNodeRef GetRootNode() const;
+		FString DumpToString(bool bIncludeInheritOriginInfo = true, bool bIncludeMaps = true) const;
+
+	private:
+		TUniquePtr<Internal::FPcpPrimIndexImpl> Impl;
+	};
+}	 // namespace UE
diff --git a/Engine/Plugins/Runtime/USDCore/Source/UnrealUSDWrapper/Public/UsdWrappers/SdfPayloadsProxy.h b/Engine/Plugins/Runtime/USDCore/Source/UnrealUSDWrapper/Public/UsdWrappers/SdfPayloadsProxy.h
new file mode 100644
--- /dev/null
+++ b/Engine/Plugins/Runtime/USDCore/Source/UnrealUSDWrapper/Public/UsdWrappers/SdfPayloadsProxy.h
@@ -0,0 +1,66 @@
+// Copyright Epic Games, Inc. All Rights Reserved.
+
+#pragma once
+
+#include "Templates/UniquePtr.h"
+
+#if USE_USD_SDK
+
+#include "USDIncludesStart.h"
+#include "pxr/pxr.h"
+#include "USDIncludesEnd.h"
+
+PXR_NAMESPACE_OPEN_SCOPE
+	class SdfPayloadTypePolicy;
+	template<typename T>
+	class SdfListEditorProxy;
+	using SdfPayloadsProxy = SdfListEditorProxy<SdfPayloadTypePolicy>;
+PXR_NAMESPACE_CLOSE_SCOPE
+
+#endif	  // #if USE_USD_SDK
+
+namespace UE
+{
+	struct FSdfPayload;
+	namespace Internal
+	{
+		class FSdfPayloadsProxyImpl;
+	}
+
+	/**
+	 * Minimal pxr::SdfPayloadsProxy wrapper for Unreal that can be used from no-rtti modules.
+	 */
+	class UNREALUSDWRAPPER_API FSdfPayloadsProxy
+	{
+	public:
+		FSdfPayloadsProxy();
+
+		FSdfPayloadsProxy(const FSdfPayloadsProxy& Other);
+		FSdfPayloadsProxy(FSdfPayloadsProxy&& Other);
+		~FSdfPayloadsProxy();
+
+		FSdfPayloadsProxy& operator=(const FSdfPayloadsProxy& Other);
+		FSdfPayloadsProxy& operator=(FSdfPayloadsProxy&& Other);
+
+		explicit operator bool() const;
+
+		// Auto conversion from/to pxr::SdfPayloadsProxy
+	public:
+#if USE_USD_SDK
+		explicit FSdfPayloadsProxy(const pxr::SdfPayloadsProxy& InSdfPayloadsProxy);
+		explicit FSdfPayloadsProxy(pxr::SdfPayloadsProxy&& InSdfPayloadsProxy);
+		FSdfPayloadsProxy& operator=(const pxr::SdfPayloadsProxy& InSdfPayloadsProxy);
+		FSdfPayloadsProxy& operator=(pxr::SdfPayloadsProxy&& InSdfPayloadsProxy);
+
+		operator pxr::SdfPayloadsProxy&();
+		operator const pxr::SdfPayloadsProxy&() const;
+#endif	  // #if USE_USD_SDK
+
+		  // Wrapped pxr::SdfPayloadsProxy functions, refer to the USD SDK documentation
+	public:
+		TArray<FSdfPayload> GetAppliedItems() const;
+
+	private:
+		TUniquePtr<Internal::FSdfPayloadsProxyImpl> Impl;
+	};
+}	 // namespace UE
diff --git a/Engine/Plugins/Runtime/USDCore/Source/UnrealUSDWrapper/Public/UsdWrappers/SdfPrimSpec.h b/Engine/Plugins/Runtime/USDCore/Source/UnrealUSDWrapper/Public/UsdWrappers/SdfPrimSpec.h
--- a/Engine/Plugins/Runtime/USDCore/Source/UnrealUSDWrapper/Public/UsdWrappers/SdfPrimSpec.h
+++ b/Engine/Plugins/Runtime/USDCore/Source/UnrealUSDWrapper/Public/UsdWrappers/SdfPrimSpec.h
@@ -22,8 +22,10 @@ PXR_NAMESPACE_CLOSE_SCOPE
 
 namespace UE
 {
-	class FSdfPath;
 	class FSdfAttributeSpec;
+	class FSdfPath;
+	class FSdfPayloadsProxy;
+	class FSdfReferencesProxy;
 
 	namespace Internal
 	{
@@ -92,6 +94,9 @@ namespace UE
 
 		UE::FSdfAttributeSpec GetAttributeAtPath(const UE::FSdfPath& Path) const;
 
+		FSdfPayloadsProxy GetPayloadList() const;
+		FSdfReferencesProxy GetReferenceList() const;
+
 		FName GetTypeName() const;
 		FName GetName() const;
 
diff --git a/Engine/Plugins/Runtime/USDCore/Source/UnrealUSDWrapper/Public/UsdWrappers/SdfReferencesProxy.h b/Engine/Plugins/Runtime/USDCore/Source/UnrealUSDWrapper/Public/UsdWrappers/SdfReferencesProxy.h
new file mode 100644
--- /dev/null
+++ b/Engine/Plugins/Runtime/USDCore/Source/UnrealUSDWrapper/Public/UsdWrappers/SdfReferencesProxy.h
@@ -0,0 +1,66 @@
+// Copyright Epic Games, Inc. All Rights Reserved.
+
+#pragma once
+
+#include "Templates/UniquePtr.h"
+
+#if USE_USD_SDK
+
+#include "USDIncludesStart.h"
+#include "pxr/pxr.h"
+#include "USDIncludesEnd.h"
+
+PXR_NAMESPACE_OPEN_SCOPE
+	class SdfReferenceTypePolicy;
+	template<typename T>
+	class SdfListEditorProxy;
+	using SdfReferencesProxy = SdfListEditorProxy<SdfReferenceTypePolicy>;
+PXR_NAMESPACE_CLOSE_SCOPE
+
+#endif	  // #if USE_USD_SDK
+
+namespace UE
+{
+	struct FSdfReference;
+	namespace Internal
+	{
+		class FSdfReferencesProxyImpl;
+	}
+
+	/**
+	 * Minimal pxr::SdfReferencesProxy wrapper for Unreal that can be used from no-rtti modules.
+	 */
+	class UNREALUSDWRAPPER_API FSdfReferencesProxy
+	{
+	public:
+		FSdfReferencesProxy();
+
+		FSdfReferencesProxy(const FSdfReferencesProxy& Other);
+		FSdfReferencesProxy(FSdfReferencesProxy&& Other);
+		~FSdfReferencesProxy();
+
+		FSdfReferencesProxy& operator=(const FSdfReferencesProxy& Other);
+		FSdfReferencesProxy& operator=(FSdfReferencesProxy&& Other);
+
+		explicit operator bool() const;
+
+		// Auto conversion from/to pxr::SdfReferencesProxy
+	public:
+#if USE_USD_SDK
+		explicit FSdfReferencesProxy(const pxr::SdfReferencesProxy& InSdfReferencesProxy);
+		explicit FSdfReferencesProxy(pxr::SdfReferencesProxy&& InSdfReferencesProxy);
+		FSdfReferencesProxy& operator=(const pxr::SdfReferencesProxy& InSdfReferencesProxy);
+		FSdfReferencesProxy& operator=(pxr::SdfReferencesProxy&& InSdfReferencesProxy);
+
+		operator pxr::SdfReferencesProxy&();
+		operator const pxr::SdfReferencesProxy&() const;
+#endif	  // #if USE_USD_SDK
+
+		  // Wrapped pxr::SdfReferencesProxy functions, refer to the USD SDK documentation
+	public:
+		TArray<FSdfReference> GetAppliedItems() const;
+
+	private:
+		TUniquePtr<Internal::FSdfReferencesProxyImpl> Impl;
+	};
+}	 // namespace UE
diff --git a/Engine/Plugins/Runtime/USDCore/Source/UnrealUSDWrapper/Public/UsdWrappers/UsdAttribute.h b/Engine/Plugins/Runtime/USDCore/Source/UnrealUSDWrapper/Public/UsdWrappers/UsdAttribute.h
--- a/Engine/Plugins/Runtime/USDCore/Source/UnrealUSDWrapper/Public/UsdWrappers/UsdAttribute.h
+++ b/Engine/Plugins/Runtime/USDCore/Source/UnrealUSDWrapper/Public/UsdWrappers/UsdAttribute.h
@@ -22,6 +22,7 @@ namespace UE
 	class FSdfPath;
 	class FUsdPrim;
 	class FVtValue;
+	class FUsdResolveInfo;
 
 	namespace Internal
 	{
@@ -90,6 +91,9 @@ namespace UE
 		bool ValueMightBeTimeVarying() const;
 
 		bool Get(UE::FVtValue& Value, TOptional<double> Time = {}) const;
+
+		FUsdResolveInfo GetResolveInfo(double Time) const;
+
 		bool Set(const UE::FVtValue& Value, TOptional<double> Time = {}) const;
 
 		template <typename T>
diff --git a/Engine/Plugins/Runtime/USDCore/Source/UnrealUSDWrapper/Public/UsdWrappers/UsdPrim.h b/Engine/Plugins/Runtime/USDCore/Source/UnrealUSDWrapper/Public/UsdWrappers/UsdPrim.h
--- a/Engine/Plugins/Runtime/USDCore/Source/UnrealUSDWrapper/Public/UsdWrappers/UsdPrim.h
+++ b/Engine/Plugins/Runtime/USDCore/Source/UnrealUSDWrapper/Public/UsdWrappers/UsdPrim.h
@@ -22,13 +22,15 @@ PXR_NAMESPACE_CLOSE_SCOPE
 
 namespace UE
 {
+	class FPcpPrimIndex;
 	class FSdfPath;
+	class FSdfPrimSpec;
 	class FUsdAttribute;
 	class FUsdPayloads;
 	class FUsdReferences;
+	class FUsdRelationship;
 	class FUsdVariantSet;
 	class FUsdVariantSets;
-	class FUsdRelationship;
 	class FVtValue;
 
 	namespace Internal
@@ -148,6 +150,11 @@ namespace UE
 		FUsdReferences GetReferences() const;
 		bool HasAuthoredReferences() const;
 
+		TArray<FSdfPrimSpec> GetPrimStack() const;
+
+		FPcpPrimIndex GetPrimIndex() const;
+		FPcpPrimIndex ComputeExpandedPrimIndex() const;
+
 		bool RemoveProperty(FName PropName) const;
 
 		FUsdAttribute CreateAttribute(const TCHAR* AttrName, FName TypeName) const;
diff --git a/Engine/Plugins/Runtime/USDCore/Source/UnrealUSDWrapper/Public/UsdWrappers/UsdResolveInfo.h b/Engine/Plugins/Runtime/USDCore/Source/UnrealUSDWrapper/Public/UsdWrappers/UsdResolveInfo.h
new file mode 100644
--- /dev/null
+++ b/Engine/Plugins/Runtime/USDCore/Source/UnrealUSDWrapper/Public/UsdWrappers/UsdResolveInfo.h
@@ -0,0 +1,57 @@
+// Copyright Epic Games, Inc. All Rights Reserved.
+
+#pragma once
+
+#include "CoreTypes.h"
+#include "Templates/UniquePtr.h"
+
+#if USE_USD_SDK
+
+#include "USDIncludesStart.h"
+#include "pxr/pxr.h"
+#include "USDIncludesEnd.h"
+
+PXR_NAMESPACE_OPEN_SCOPE
+	class UsdResolveInfo;
+PXR_NAMESPACE_CLOSE_SCOPE
+
+#endif	  // #if USE_USD_SDK
+
+namespace UE
+{
+	class FPcpNodeRef;
+	namespace Internal
+	{
+		class FUsdResolveInfoImpl;
+	}
+
+	/**
+	 * Minimal pxr::UsdResolveInfo wrapper for Unreal that can be used from no-rtti modules.
+	 */
+	class UNREALUSDWRAPPER_API FUsdResolveInfo
+	{
+	public:
+		FUsdResolveInfo();
+		FUsdResolveInfo(const FUsdResolveInfo& Other);
+		FUsdResolveInfo(FUsdResolveInfo&& Other);
+
+		~FUsdResolveInfo();
+
+		// Auto conversion from/to pxr::UsdResolveInfo
+	public:
+#if USE_USD_SDK
+		explicit FUsdResolveInfo(const pxr::UsdResolveInfo& InResolveInfo);
+		explicit FUsdResolveInfo(pxr::UsdResolveInfo&& InResolveInfo);
+
+		operator pxr::UsdResolveInfo&();
+		operator const pxr::UsdResolveInfo&() const;
+#endif	  // #if USE_USD_SDK
+
+		  // Wrapped pxr::UsdResolveInfo functions, refer to the USD SDK documentation
+	public:
+		FPcpNodeRef GetNode() const;
+
+	private:
+		TUniquePtr<Internal::FUsdResolveInfoImpl> Impl;
+	};
+}	 // namespace UE
