USD:
- Add support for volume shader graph with MaterialX
- Surround MaterialX code by a try catch to avoid any potential exception and crash

USDUtilities
- Make MaterialX module private

diff --git a/Engine/Plugins/Interchange/Extensions/OpenUSD/Source/Import/Private/InterchangeUsdTranslator.cpp b/Engine/Plugins/Interchange/Extensions/OpenUSD/Source/Import/Private/InterchangeUsdTranslator.cpp
--- a/Engine/Plugins/Interchange/Extensions/OpenUSD/Source/Import/Private/InterchangeUsdTranslator.cpp
+++ b/Engine/Plugins/Interchange/Extensions/OpenUSD/Source/Import/Private/InterchangeUsdTranslator.cpp
@@ -1152,7 +1152,10 @@ namespace UE::InterchangeUsdTranslator::Private
 
 		// Check for any references of MaterialX
 #if WITH_EDITOR
-		if (RenderContext == UnrealIdentifiers::MaterialXRenderContext && UsdUtils::HasSurfaceOutput(Prim, UnrealIdentifiers::MaterialXRenderContext))
+		if (RenderContext == UnrealIdentifiers::MaterialXRenderContext &&
+			(UsdUtils::HasSurfaceOutput(Prim, UnrealIdentifiers::MaterialXRenderContext) ||
+			UsdUtils::HasVolumeOutput(Prim, UnrealIdentifiers::MaterialXRenderContext) ||
+			UsdUtils::HasDisplacementOutput(Prim, UnrealIdentifiers::MaterialXRenderContext)))
 		{
 			TArray<FString> FilePaths = UsdUtils::GetMaterialXFilePaths(Prim);
 			for (const FString& File : FilePaths)
diff --git a/Engine/Plugins/Runtime/USDCore/Source/USDUtilities/Private/USDMaterialXShaderGraph.cpp b/Engine/Plugins/Runtime/USDCore/Source/USDUtilities/Private/USDMaterialXShaderGraph.cpp
--- a/Engine/Plugins/Runtime/USDCore/Source/USDUtilities/Private/USDMaterialXShaderGraph.cpp
+++ b/Engine/Plugins/Runtime/USDCore/Source/USDUtilities/Private/USDMaterialXShaderGraph.cpp
@@ -11,6 +11,8 @@
 #include "USDTypesConversion.h"
 #include "UsdWrappers/SdfLayer.h"
 
+#define LOCTEXT_NAMESPACE "USDMaterialXShaderGraph"
+
 namespace mx = MaterialX;
 
 FUsdMaterialXShaderGraph::FUsdMaterialXShaderGraph(const pxr::UsdPrim& UsdShadeMaterialPrim, const TCHAR* RenderContext)
@@ -70,17 +72,25 @@ FUsdMaterialXShaderGraph::FUsdMaterialXShaderGraph(const pxr::UsdPrim& UsdShadeM
 
 	if (!UsdShadeMaterial)
 	{
-		USD_LOG_ERROR(TEXT("Couldn't create a UsdShadeMaterial for the prim."));
+		USD_LOG_USERERROR(FText::Format(LOCTEXT("MaterialXShaderGraph", "Couldn't create a UsdShadeMaterial for the prim {0}"),
+						  FText::FromString(UsdToUnreal::ConvertToken(UsdShadeMaterialPrim.GetName()))));
 		return;
 	}
 
-	Document = mx::createDocument();
-	mx::DocumentPtr MaterialXLibrary = mx::createDocument();
-	mx::FileSearchPath MaterialXFolder{TCHAR_TO_UTF8(*FPaths::Combine(FPaths::EngineDir(), TEXT("Binaries"), TEXT("ThirdParty"), TEXT("MaterialX")))};
-	mx::StringSet LoadedLibs = mx::loadLibraries({"libraries"}, MaterialXFolder, MaterialXLibrary);
-	Document->importLibrary(MaterialXLibrary);
-
-	CreateSurfaceMaterial(UsdShadeMaterial, RenderContext);
+	try
+	{
+		Document = mx::createDocument();
+		mx::DocumentPtr MaterialXLibrary = mx::createDocument();
+		mx::FileSearchPath MaterialXFolder{ TCHAR_TO_UTF8(*FPaths::Combine(FPaths::EngineDir(), TEXT("Binaries"), TEXT("ThirdParty"), TEXT("MaterialX"))) };
+		mx::StringSet LoadedLibs = mx::loadLibraries({ "libraries" }, MaterialXFolder, MaterialXLibrary);
+		Document->importLibrary(MaterialXLibrary);
+		CreateMaterial(UsdShadeMaterial, RenderContext);
+	}
+	catch (const std::exception& Exception)
+	{
+		USD_LOG_USERERROR(FText::Format(LOCTEXT("MaterialXShaderGraph", "MaterialX: {0}\n shader graph creation aborted...\n"),
+						  FText::FromString(Exception.what())));
+	}
 }
 
 mx::DocumentPtr FUsdMaterialXShaderGraph::GetDocument() const
@@ -93,7 +103,7 @@ const TArray<FUsdMaterialXShaderGraph::FGeomProp>& FUsdMaterialXShaderGraph::Get
 	return GeomPropValueNames;
 }
 
-void FUsdMaterialXShaderGraph::CreateSurfaceMaterial(const pxr::UsdShadeMaterial& UsdShadeMaterial, const TCHAR* RenderContext)
+void FUsdMaterialXShaderGraph::CreateMaterial(const pxr::UsdShadeMaterial& UsdShadeMaterial, const TCHAR* RenderContext)
 {
 	pxr::TfToken RenderContextToken = pxr::UsdShadeTokens->universalRenderContext;
 	if (RenderContext)
@@ -102,22 +112,23 @@ void FUsdMaterialXShaderGraph::CreateSurfaceMaterial(const pxr::UsdShadeMaterial
 		RenderContextToken = ProvidedRenderContextToken;
 	}
 
-	pxr::UsdShadeShader SurfaceShader = UsdShadeMaterial.ComputeSurfaceSource(RenderContextToken);
-	pxr::UsdShadeShader DisplacementShader = UsdShadeMaterial.ComputeDisplacementSource(RenderContextToken);
+	pxr::UsdShadeShader SurfaceSource = UsdShadeMaterial.ComputeSurfaceSource(RenderContextToken);
+	pxr::UsdShadeShader DisplacementSource = UsdShadeMaterial.ComputeDisplacementSource(RenderContextToken);
+	pxr::UsdShadeShader VolumeSource = UsdShadeMaterial.ComputeVolumeSource(RenderContextToken);
 
-	if (!SurfaceShader)
+	if (!SurfaceSource && !VolumeSource)
 	{
 		USD_LOG_ERROR(
-			TEXT("Couldn't find a SurfaceShader for UsdShadeMaterial '%s'."),
+			TEXT("Couldn't find any surface or volume shaders for the UsdShadeMaterial: '%s'."),
 			*UsdToUnreal::ConvertToken(UsdShadeMaterial.GetPrim().GetName())
 		);
 		return;
 	}
 
-	auto CreateShaderGraph = [this](const pxr::UsdShadeShader& SurfaceShader, mx::DocumentPtr InDocument) -> mx::NodePtr
+	auto CreateShaderGraph = [this](const pxr::UsdShadeShader& ShadeShader, mx::DocumentPtr InDocument) -> mx::NodePtr
 	{
 		pxr::TfToken NodeDefTokenId;
-		SurfaceShader.GetShaderId(&NodeDefTokenId);
+		ShadeShader.GetShaderId(&NodeDefTokenId);
 		mx::NodeDefPtr NodeDef = InDocument->getNodeDef(NodeDefTokenId.GetString());
 
 		// no node definition, just return
@@ -126,15 +137,15 @@ void FUsdMaterialXShaderGraph::CreateSurfaceMaterial(const pxr::UsdShadeMaterial
 			return nullptr;
 		}
 
-		mx::NodePtr MtlxSurfaceShader = InDocument->addNodeInstance(NodeDef, SurfaceShader.GetPrim().GetName().GetString());
+		mx::NodePtr Shader = InDocument->addNodeInstance(NodeDef, ShadeShader.GetPrim().GetName().GetString());
 
 		// if the shadergraph is not valid (probably caused by a missing nodedef or wrong type), just invalidate it all
-		if (!ComputeShaderGraph(MtlxSurfaceShader, SurfaceShader))
+		if (!ComputeShaderGraph(Shader, ShadeShader))
 		{
-			MtlxSurfaceShader = nullptr;
+			Shader = nullptr;
 		}
 
-		return MtlxSurfaceShader;
+		return Shader;
 	};
 
 	// Create and store the input interface names
@@ -151,31 +162,75 @@ void FUsdMaterialXShaderGraph::CreateSurfaceMaterial(const pxr::UsdShadeMaterial
 		InterfaceNames.insert(InputInterfaceName->getName());
 	}
 
-	// Create the shader graph for the surface shader and displacement shader
-	if (mx::NodePtr MaterialXSurfaceShader = CreateShaderGraph(SurfaceShader, Document))
+	// Create the shader graph for the surface/volume/displacement shader
+	// for the following source materials since we create them after the shader graph, we need to make sure that their name is unique
+	// unlike MaterialX, USD allows a same name between a parent and its child
+	std::string MaterialName = UsdShadeMaterial.GetPrim().GetName().GetString();
+	mx::NodePtr SurfaceMaterial;
+	if(SurfaceSource)
 	{
-		mx::NodeDefPtr SurfaceMaterialNodeDef = Document->getNodeDef("ND_surfacematerial");
-		mx::NodePtr SurfaceMaterial = Document->addNodeInstance(SurfaceMaterialNodeDef, UsdShadeMaterial.GetPrim().GetName().GetString());
+		mx::NodePtr SurfaceShader = CreateShaderGraph(SurfaceSource, Document);
+		if (SurfaceShader)
+		{
+			// create the surfacematerial first
+			mx::NodeDefPtr SurfaceMaterialNodeDef = Document->getNodeDef("ND_surfacematerial");
+			std::string SurfaceMaterialName = Document->getChild(MaterialName) ? Document->createValidChildName(MaterialName) : MaterialName;
+			SurfaceMaterial = Document->addNodeInstance(SurfaceMaterialNodeDef, SurfaceMaterialName);
 
-		SurfaceMaterial->setConnectedNode("surfaceshader", MaterialXSurfaceShader);
+			SurfaceMaterial->setConnectedNode("surfaceshader", SurfaceShader);
+		}
+	}
 
-		// the displacement shader is not necessarily required
-		if (DisplacementShader)
+	// a displacement shader is connected to a surfacematerial, we have to check if we created it before
+	if (DisplacementSource)
+	{
+		if (mx::NodePtr DisplacementShader = CreateShaderGraph(DisplacementSource, Document))
 		{
-			if (mx::NodePtr MaterialXDisplacementShader = CreateShaderGraph(DisplacementShader, Document))
+			if (!SurfaceMaterial)
 			{
-				SurfaceMaterial->setConnectedNode("displacementshader", MaterialXDisplacementShader);
+				mx::NodeDefPtr SurfaceMaterialNodeDef = Document->getNodeDef("ND_surfacematerial");
+				std::string DisplacementMaterialName = Document->getChild(MaterialName) ? Document->createValidChildName(MaterialName) : MaterialName;
+				SurfaceMaterial = Document->addNodeInstance(SurfaceMaterialNodeDef, DisplacementMaterialName);
 			}
+			SurfaceMaterial->setConnectedNode("displacementshader", DisplacementShader);
 		}
+	}
 
+	if(SurfaceMaterial)
+	{
 		std::string ErrorMessage;
 		SurfaceMaterial->validate(&ErrorMessage);
 		if (!ErrorMessage.empty())
 		{
 			USD_LOG_ERROR(TEXT("USD MaterialX: %s."), UTF8_TO_TCHAR(ErrorMessage.c_str()));
+			Document->removeNode(UsdShadeMaterial.GetPrim().GetName().GetString());
 			SurfaceMaterial = nullptr;
 		}
 	}
+
+	// a volume material is another type of material different from a surface material, it's unclear if we can have both a surface and volume in the same Material
+	// if that's the case let's just suffix the name of the volume material to avoid collision during the creation of the node
+	if(VolumeSource)
+	{
+		if (mx::NodePtr VolumeShader = CreateShaderGraph(VolumeSource, Document))
+		{
+			mx::NodeDefPtr VolumeMaterialNodeDef = Document->getNodeDef("ND_volumematerial");
+			std::string VolumeMaterialName = Document->getChild(MaterialName) ? Document->createValidChildName(MaterialName) : MaterialName;
+			mx::NodePtr VolumeMaterial = Document->addNodeInstance(VolumeMaterialNodeDef, VolumeMaterialName);
+
+			VolumeMaterial->setConnectedNode("volumeshader", VolumeShader);
+
+			std::string ErrorVolumeMessage;
+			VolumeMaterial->validate(&ErrorVolumeMessage);
+			if (!ErrorVolumeMessage.empty())
+			{
+				USD_LOG_ERROR(TEXT("USD MaterialX: %s."), UTF8_TO_TCHAR(ErrorVolumeMessage.c_str()));
+				Document->removeNode(VolumeMaterialName);
+				VolumeMaterial = nullptr;
+			}
+		}
+	}
+
 }
 
 bool FUsdMaterialXShaderGraph::ComputeShaderGraph(mx::NodePtr Node, const pxr::UsdShadeShader& Shader)
@@ -540,4 +595,6 @@ bool FUsdMaterialXShaderGraph::IsTangentSpaceInput(const pxr::UsdShadeInput& Usd
 {
 	return TangentSpaceInputs.find(UsdInput.GetBaseName().GetString()) != TangentSpaceInputs.cend();
 }
+
+#undef LOCTEXT_NAMESPACE 
 #endif // USE_USD_SDK && ENABLE_USD_MATERIALX
diff --git a/Engine/Plugins/Runtime/USDCore/Source/USDUtilities/Private/USDShadeConversion.cpp b/Engine/Plugins/Runtime/USDCore/Source/USDUtilities/Private/USDShadeConversion.cpp
--- a/Engine/Plugins/Runtime/USDCore/Source/USDUtilities/Private/USDShadeConversion.cpp
+++ b/Engine/Plugins/Runtime/USDCore/Source/USDUtilities/Private/USDShadeConversion.cpp
@@ -63,12 +63,6 @@
 
 #include "USDIncludesEnd.h"
 
-#if WITH_EDITOR && PLATFORM_WINDOWS
-#include "MaterialXFormat/Util.h"
-
-namespace mx = MaterialX;
-#endif
-
 // Required for packaging
 class UTextureFactory;
 
@@ -3468,6 +3462,44 @@ bool UsdUtils::HasSurfaceOutput(const pxr::UsdPrim& MaterialPrim, const FName& R
 	return static_cast<bool>(MaterialOutput);
 }
 
+bool UsdUtils::HasDisplacementOutput(const pxr::UsdPrim& MaterialPrim, const FName& RenderContext)
+{
+	if (!MaterialPrim)
+	{
+		return false;
+	}
+
+	FScopedUsdAllocs UsdAllocs;
+
+	pxr::UsdShadeMaterial ShadeMaterial{MaterialPrim};
+	if (!ShadeMaterial)
+	{
+		return false;
+	}
+
+	pxr::UsdShadeOutput MaterialOutput = ShadeMaterial.GetDisplacementOutput(UnrealToUsd::ConvertToken(*RenderContext.ToString()).Get());
+	return static_cast<bool>(MaterialOutput);
+}
+
+bool UsdUtils::HasVolumeOutput(const pxr::UsdPrim& MaterialPrim, const FName& RenderContext)
+{
+	if (!MaterialPrim)
+	{
+		return false;
+	}
+
+	FScopedUsdAllocs UsdAllocs;
+
+	pxr::UsdShadeMaterial ShadeMaterial{MaterialPrim};
+	if (!ShadeMaterial)
+	{
+		return false;
+	}
+
+	pxr::UsdShadeOutput MaterialOutput = ShadeMaterial.GetVolumeOutput(UnrealToUsd::ConvertToken(*RenderContext.ToString()).Get());
+	return static_cast<bool>(MaterialOutput);
+}
+
 bool UsdUtils::IsMaterialTranslucent(const pxr::UsdShadeMaterial& UsdShadeMaterial)
 {
 	FScopedUsdAllocs UsdAllocs;
diff --git a/Engine/Plugins/Runtime/USDCore/Source/USDUtilities/Public/USDMaterialXShaderGraph.h b/Engine/Plugins/Runtime/USDCore/Source/USDUtilities/Public/USDMaterialXShaderGraph.h
--- a/Engine/Plugins/Runtime/USDCore/Source/USDUtilities/Public/USDMaterialXShaderGraph.h
+++ b/Engine/Plugins/Runtime/USDCore/Source/USDUtilities/Public/USDMaterialXShaderGraph.h
@@ -43,16 +43,17 @@ public:
 
 private:
 
-	USDUTILITIES_API void CreateSurfaceMaterial(const pxr::UsdShadeMaterial& UsdShadeMaterial, const TCHAR* RenderContext);
+	/** Create either a <surfacematerial> or a <volumematerial> if possible. */
+	void CreateMaterial(const pxr::UsdShadeMaterial& UsdShadeMaterial, const TCHAR* RenderContext);
 
-	USDUTILITIES_API bool ComputeShaderGraph(MaterialX::NodePtr Node, const pxr::UsdShadeShader& Shader);
+	bool ComputeShaderGraph(MaterialX::NodePtr Node, const pxr::UsdShadeShader& Shader);
 
-	USDUTILITIES_API bool IsGeomColor(const pxr::UsdShadeShader& GeomPropValueShader);
+	bool IsGeomColor(const pxr::UsdShadeShader& GeomPropValueShader);
 
-	USDUTILITIES_API void SetInputValue(MaterialX::InputPtr Input, const pxr::UsdShadeInput& UsdInput);
+	void SetInputValue(MaterialX::InputPtr Input, const pxr::UsdShadeInput& UsdInput);
 
 	/** Set the interfaceName of the input if any */ 
-	USDUTILITIES_API bool SetInterfaceName(MaterialX::InputPtr Input, const pxr::UsdShadeConnectionSourceInfo& Source);
+	bool SetInterfaceName(MaterialX::InputPtr Input, const pxr::UsdShadeConnectionSourceInfo& Source);
 
 	/**
 	 * Get the MaterialX node from a UsdShadeConnectionSourceInfo
@@ -60,11 +61,11 @@ private:
 	 * @param bIsNodeAlreadyExist - return true if the node was already present in the shader graph, false otherwise
 	 * @return The corresponding MaterialX node in the shader graph, if it doesn't exist it creates one
 	 */
-	USDUTILITIES_API MaterialX::NodePtr GetNode(const pxr::UsdShadeConnectionSourceInfo& Source, bool& bIsNodeAlreadyExist);
+	MaterialX::NodePtr GetNode(const pxr::UsdShadeConnectionSourceInfo& Source, bool& bIsNodeAlreadyExist);
 
-	USDUTILITIES_API MaterialX::InputPtr GetInput(MaterialX::NodePtr Node, const pxr::UsdShadeInput& UsdInput);
+	MaterialX::InputPtr GetInput(MaterialX::NodePtr Node, const pxr::UsdShadeInput& UsdInput);
 
-	USDUTILITIES_API MaterialX::NodeDefPtr GetNodeDef(const pxr::UsdShadeConnectionSourceInfo& Source, bool& bIsGeomPropValue);
+	MaterialX::NodeDefPtr GetNodeDef(const pxr::UsdShadeConnectionSourceInfo& Source, bool& bIsGeomPropValue);
 
 	/**
 	 * Given an input type, retrieve the MaterialX associated type
@@ -72,9 +73,9 @@ private:
 	 * @param InOutInputType - The type of the input, if it's valid it will now be the MaterialX type
 	 * @return true if there is a valid match between USD and MaterialX or if it's already a MaterialX type, false otherwise
 	 */
-	USDUTILITIES_API bool GetMatchingInputType(const std::string& InputName, std::string& InOutInputType);
+	bool GetMatchingInputType(const std::string& InputName, std::string& InOutInputType);
 
-	USDUTILITIES_API bool IsTangentSpaceInput(const pxr::UsdShadeInput& UsdInput) const;
+	bool IsTangentSpaceInput(const pxr::UsdShadeInput& UsdInput) const;
 
 private:
 
diff --git a/Engine/Plugins/Runtime/USDCore/Source/USDUtilities/Public/USDShadeConversion.h b/Engine/Plugins/Runtime/USDCore/Source/USDUtilities/Public/USDShadeConversion.h
--- a/Engine/Plugins/Runtime/USDCore/Source/USDUtilities/Public/USDShadeConversion.h
+++ b/Engine/Plugins/Runtime/USDCore/Source/USDUtilities/Public/USDShadeConversion.h
@@ -20,10 +20,6 @@
 #include "pxr/usd/usdShade/tokens.h"
 #include "USDIncludesEnd.h"
 
-#if WITH_EDITOR && PLATFORM_WINDOWS
-#include "MaterialXCore/Document.h"
-#endif
-
 PXR_NAMESPACE_OPEN_SCOPE
 	class UsdAttribute;
 	class UsdGeomMesh;
@@ -279,6 +275,12 @@ namespace UsdUtils
 	// Returns whether MaterialPrim is an actual Material, and has a surface output authored for the provided render context
 	USDUTILITIES_API bool HasSurfaceOutput(const pxr::UsdPrim& MaterialPrim, const FName& RenderContext);
 
+	// Returns whether MaterialPrim is an actual Material, and has a displacement output authored for the provided render context
+	USDUTILITIES_API bool HasDisplacementOutput(const pxr::UsdPrim& MaterialPrim, const FName& RenderContext);
+
+	// Returns whether MaterialPrim is an actual Material, and has a volume output authored for the provided render context
+	USDUTILITIES_API bool HasVolumeOutput(const pxr::UsdPrim& MaterialPrim, const FName& RenderContext);
+
 	/**
 	 * Returns whether the material needs to be rendered with the Translucent rendering mode.
 	 * This function exists because we need this information *before* we pick the right parent for a material instance and properly convert it.
diff --git a/Engine/Plugins/Runtime/USDCore/Source/USDUtilities/USDUtilities.Build.cs b/Engine/Plugins/Runtime/USDCore/Source/USDUtilities/USDUtilities.Build.cs
--- a/Engine/Plugins/Runtime/USDCore/Source/USDUtilities/USDUtilities.Build.cs
+++ b/Engine/Plugins/Runtime/USDCore/Source/USDUtilities/USDUtilities.Build.cs
@@ -70,9 +70,7 @@ namespace UnrealBuildTool.Rules
 				// we work through issues affecting Mac and Linux.
 				if (Target.Platform == UnrealTargetPlatform.Win64)
 				{
-					// Public because MaterialX headers are #include'd in
-					// public USDUtilities headers.
-					PublicDependencyModuleNames.AddRange(
+					PrivateDependencyModuleNames.AddRange(
 						new string[]
 						{
 							"MaterialX"
