USD:
- Fix compile error because exceptions are disabled

diff --git a/Engine/Plugins/Runtime/USDCore/Source/USDUtilities/USDUtilities.Build.cs b/Engine/Plugins/Runtime/USDCore/Source/USDUtilities/USDUtilities.Build.cs
--- a/Engine/Plugins/Runtime/USDCore/Source/USDUtilities/USDUtilities.Build.cs
+++ b/Engine/Plugins/Runtime/USDCore/Source/USDUtilities/USDUtilities.Build.cs
@@ -52,7 +52,8 @@ namespace UnrealBuildTool.Rules
 
 			if (Target.bBuildEditor)
 			{
-				PrivateDependencyModuleNames.AddRange(
+                bEnableExceptions = true;
+                PrivateDependencyModuleNames.AddRange(
 					new string[]
 					{
 						"DesktopPlatform", // For OpenFileDialog/SaveFileDialog
