USD: Add support for subdividing meshes when importing via USD Interchange.

diff --git a/Engine/Plugins/Interchange/Extensions/OpenUSD/Source/Import/Private/InterchangeUSDPipeline.cpp b/Engine/Plugins/Interchange/Extensions/OpenUSD/Source/Import/Private/InterchangeUSDPipeline.cpp
--- a/Engine/Plugins/Interchange/Extensions/OpenUSD/Source/Import/Private/InterchangeUSDPipeline.cpp
+++ b/Engine/Plugins/Interchange/Extensions/OpenUSD/Source/Import/Private/InterchangeUSDPipeline.cpp
@@ -2,6 +2,7 @@
 
 #include "InterchangeUSDPipeline.h"
 
+#include "InterchangeGeometryCacheFactoryNode.h"
 #include "InterchangeHeterogeneousVolumeActorFactoryNode.h"
 #include "InterchangeMaterialDefinitions.h"
 #include "InterchangeMaterialFactoryNode.h"
@@ -834,6 +835,21 @@ namespace UE::InterchangeUsdPipeline::Private
 	{
 		using namespace UE::Interchange;
 
+		// TODO: This doesn't work yet, and would produce a ton of extra broken "compatible" materials if unchecked.
+		// The root issue being that the PrimvarToUVIndex map made for the Mesh prim doesn't consider how the mesh data
+		// will be combined into a single Static/SkeletalMesh asset. Different Mesh prims may claim different mappings,
+		// and it could be that none match the final PrimvarToUVIndex layout we actually end up with...
+		//
+		// Even more concerning: For now, Interchange currently will combine meshes without caring about their PrimvarToUVIndex
+		// layout in general, so two cubes may have different primvars assigned to UV0, that end up merged together.
+		// Now a material that wants to read this primvar will read unwanted data on half the combined mesh! On legacy USD
+		// we did a pre-pass to organize which primvar belonged to which UV index in the combined mesh. On Interchange
+		// this doesn't make sense, as the combining is dictated by (potentially custom) pipelines that will run only later
+		if (InMeshNode->IsSkinnedMesh())
+		{
+			return;
+		}
+
 		// Rebuild the mesh's primvar to UV index map
 		TMap<FString, int32> MeshPrimvarToUVIndex;
 		{
@@ -947,7 +963,12 @@ namespace UE::InterchangeUsdPipeline::Private
 		}
 	}
 
-	void ProcessMeshNodes(UInterchangeBaseNodeContainer* NodeContainer, bool bGenerateCompatibleMaterials, EInterchangeUsdPrimvar ImportPrimvar)
+	void ProcessMeshNodes(
+		UInterchangeBaseNodeContainer* NodeContainer,
+		bool bGenerateCompatibleMaterials,
+		EInterchangeUsdPrimvar ImportPrimvar,
+		int32 SubdivisionLevel
+	)
 	{
 		using namespace UE::Interchange;
 
@@ -956,63 +977,89 @@ namespace UE::InterchangeUsdPipeline::Private
 		TMap<const UInterchangeMeshNode*, UInterchangeMeshFactoryNode*> MeshNodeToMeshFactoryNode;
 		TMap<const UInterchangeMeshNode*, UInterchangeMeshActorFactoryNode*> MeshNodeToMeshActorFactoryNode;
 
-		NodeContainer->IterateNodesOfType<UInterchangeMeshNode>(
-			[NodeContainer, &MeshNodeToMeshFactoryNode, ImportPrimvar](const FString& NodeUid, UInterchangeMeshNode* MeshNode)
+		NodeContainer->IterateNodes(
+			[NodeContainer,
+			 &MeshNodeToMeshFactoryNode,
+			 ImportPrimvar,
+			 SubdivisionLevel,
+			 bGenerateCompatibleMaterials,
+			 &MeshNodeToMeshActorFactoryNode](const FString& NodeUid, UInterchangeBaseNode* BaseNode)
 			{
-				const FString FactoryNodeUid = UInterchangeFactoryBaseNode::BuildFactoryNodeUid(NodeUid);
-				UInterchangeMeshFactoryNode* MeshFactoryNode = Cast<UInterchangeMeshFactoryNode>(NodeContainer->GetFactoryNode(FactoryNodeUid));
-				if (!MeshFactoryNode)
+				if (UInterchangeMeshNode* MeshNode = Cast<UInterchangeMeshNode>(BaseNode))
 				{
-					return;
-				}
-				MeshNodeToMeshFactoryNode.Add(MeshNode, MeshFactoryNode);
-
-#if WITH_EDITORONLY_DATA
-				MeshFactoryNode->AddPayloadKeyInt32Attribute(UE::Interchange::USD::Primvar::Import, int32(ImportPrimvar));
-				if (ImportPrimvar != EInterchangeUsdPrimvar::Standard)
-				{
-					AddPrimvarAttributesToFactoryNodes(MeshNode, MeshFactoryNode, NodeContainer);
-				}
-#endif	  // WITH_EDITORONLY_DATA
-			}
-		);
+					const FString FactoryNodeUid = UInterchangeFactoryBaseNode::BuildFactoryNodeUid(NodeUid);
+					UInterchangeMeshFactoryNode* MeshFactoryNode = Cast<UInterchangeMeshFactoryNode>(NodeContainer->GetFactoryNode(FactoryNodeUid));
+					if (!MeshFactoryNode)
+					{
+						if (!MeshFactoryNode)
+						{
+							return;
+						}
+					}
 
-		if (bGenerateCompatibleMaterials)
-		{
-			// Also search through our mesh actor factory nodes: We may have added material overrides that also need to be made compatible.
-			//
-			// Two passes through the node container is not ideal, but at least we can turn off this whole pass by disabling
-			// bGenerateCompatibleMaterials, as opposed to having the collection of these nodes always running on the above pass
-			NodeContainer->IterateNodesOfType<UInterchangeMeshActorFactoryNode>(
-				[NodeContainer, &MeshNodeToMeshActorFactoryNode](const FString& NodeUid, UInterchangeMeshActorFactoryNode* MeshActorFactoryNode)
-				{
-					TMap<FString, FString> MaterialDependencies;
-					MeshActorFactoryNode->GetSlotMaterialDependencies(MaterialDependencies);
-					if (MaterialDependencies.Num() == 0)
+					if (bGenerateCompatibleMaterials)
 					{
-						return;
+						MeshNodeToMeshFactoryNode.Add(MeshNode, MeshFactoryNode);
 					}
 
-					FString MeshFactoryNodeUid;
-					MeshActorFactoryNode->GetCustomInstancedAssetFactoryNodeUid(MeshFactoryNodeUid);
-					UInterchangeMeshFactoryNode* MeshFactoryNode = Cast<UInterchangeMeshFactoryNode>(NodeContainer->GetFactoryNode(MeshFactoryNodeUid));
-					if (!MeshFactoryNode)
+					if (SubdivisionLevel != 0)
 					{
-						return;
+						MeshFactoryNode->AddPayloadKeyInt32Attribute(UE::Interchange::USD::SubdivisionLevelAttributeKey, SubdivisionLevel);
 					}
 
-					TArray<FString> MeshFactoryNodeTargetUids;
-					MeshFactoryNode->GetTargetNodeUids(MeshFactoryNodeTargetUids);
-					for (const FString& TargetUid : MeshFactoryNodeTargetUids)
+#if WITH_EDITORONLY_DATA
+					MeshFactoryNode->AddPayloadKeyInt32Attribute(UE::Interchange::USD::Primvar::Import, int32(ImportPrimvar));
+					if (ImportPrimvar != EInterchangeUsdPrimvar::Standard)
 					{
-						if (const UInterchangeMeshNode* MeshNode = Cast<const UInterchangeMeshNode>(NodeContainer->GetNode(TargetUid)))
+						AddPrimvarAttributesToFactoryNodes(MeshNode, MeshFactoryNode, NodeContainer);
+					}
+#endif	  // WITH_EDITORONLY_DATA
+				}
+				else if (UInterchangeGeometryCacheFactoryNode* GeometryCacheNode = Cast<UInterchangeGeometryCacheFactoryNode>(BaseNode))
+				{
+					if (SubdivisionLevel != 0) 
+					{
+						GeometryCacheNode->AddPayloadKeyInt32Attribute(UE::Interchange::USD::SubdivisionLevelAttributeKey, SubdivisionLevel);
+					}
+				}
+				else if (UInterchangeMeshActorFactoryNode* MeshActorFactoryNode = Cast<UInterchangeMeshActorFactoryNode>(BaseNode))
+				{
+					// Also search through our mesh actor factory nodes: We may have added material overrides that also need to be made compatible.
+					if (bGenerateCompatibleMaterials)
+					{
+						TMap<FString, FString> MaterialDependencies;
+						MeshActorFactoryNode->GetSlotMaterialDependencies(MaterialDependencies);
+						if (MaterialDependencies.Num() == 0)
 						{
-							MeshNodeToMeshActorFactoryNode.Add(MeshNode, MeshActorFactoryNode);
+							return;
+						}
+
+						FString MeshFactoryNodeUid;
+						MeshActorFactoryNode->GetCustomInstancedAssetFactoryNodeUid(MeshFactoryNodeUid);
+						UInterchangeMeshFactoryNode* AssignedMeshFactoryNode = Cast<UInterchangeMeshFactoryNode>(
+							NodeContainer->GetFactoryNode(MeshFactoryNodeUid)
+						);
+						if (!AssignedMeshFactoryNode)
+						{
+							return;
+						}
+
+						TArray<FString> MeshFactoryNodeTargetUids;
+						AssignedMeshFactoryNode->GetTargetNodeUids(MeshFactoryNodeTargetUids);
+						for (const FString& TargetUid : MeshFactoryNodeTargetUids)
+						{
+							if (const UInterchangeMeshNode* AssignedMeshNode = Cast<const UInterchangeMeshNode>(NodeContainer->GetNode(TargetUid)))
+							{
+								MeshNodeToMeshActorFactoryNode.Add(AssignedMeshNode, MeshActorFactoryNode);
+							}
 						}
 					}
 				}
-			);
+			}
+		);
 
+		if (bGenerateCompatibleMaterials)
+		{
 			// Build this map to prevent us from reparsing the same material in case multiple meshes are using them
 			TMap<FString, FMaterialPrimvarInfo> MaterialFactoryNodeUidToInfo;
 
@@ -1351,6 +1398,7 @@ namespace UE::InterchangeUsdPipeline::Private
 
 UInterchangeUsdPipeline::UInterchangeUsdPipeline()
 	: GeometryPurpose((int32)(EUsdPurpose::Default | EUsdPurpose::Proxy | EUsdPurpose::Render | EUsdPurpose::Guide))
+	, SubdivisionLevel(0)
 	, bImportPseudoRoot(true)
 	, bGeneratePrimvarCompatibleMaterials(true)
 {
@@ -1414,7 +1462,7 @@ void UInterchangeUsdPipeline::ExecutePipeline(
 		RemovePseudoRootFactoryNode(*NodeContainer);
 	}
 
-	ProcessMeshNodes(NodeContainer, bGeneratePrimvarCompatibleMaterials, ImportPrimvars);
+	ProcessMeshNodes(NodeContainer, bGeneratePrimvarCompatibleMaterials, ImportPrimvars, SubdivisionLevel);
 #endif	  // USE_USD_SDK
 }
 
diff --git a/Engine/Plugins/Interchange/Extensions/OpenUSD/Source/Import/Private/InterchangeUsdTranslator.cpp b/Engine/Plugins/Interchange/Extensions/OpenUSD/Source/Import/Private/InterchangeUsdTranslator.cpp
--- a/Engine/Plugins/Interchange/Extensions/OpenUSD/Source/Import/Private/InterchangeUsdTranslator.cpp
+++ b/Engine/Plugins/Interchange/Extensions/OpenUSD/Source/Import/Private/InterchangeUsdTranslator.cpp
@@ -4725,6 +4725,7 @@ TOptional<UE::Interchange::FMeshPayloadData> UInterchangeUSDTranslator::GetMeshP
 	static_assert(uint8(EInterchangeUsdPrimvar::Bake) == FUsdMeshConversionOptions::EImportPrimvar::Bake, "FUsdMeshConversionOptions::EImportPrimvar::Bake is different from EInterchangeUsdPrimvar::Bake");
 	static_assert(uint8(EInterchangeUsdPrimvar::Standard) == FUsdMeshConversionOptions::EImportPrimvar::Standard, "FUsdMeshConversionOptions::EImportPrimvar::Standard is different from EInterchangeUsdPrimvar::Standard");
 
+	PayloadAttributes.GetAttribute(UE::Interchange::FAttributeKey{ USD::SubdivisionLevelAttributeKey }, OptionsCopy.SubdivisionLevel);
 	PayloadAttributes.GetAttribute(UE::Interchange::FAttributeKey{ USD::Primvar::Import }, reinterpret_cast<int32&>(OptionsCopy.ImportPrimvars));
 
 	if (EInterchangeUsdPrimvar(OptionsCopy.ImportPrimvars) != EInterchangeUsdPrimvar::Standard)
diff --git a/Engine/Plugins/Interchange/Extensions/OpenUSD/Source/Import/Public/InterchangeUSDPipeline.h b/Engine/Plugins/Interchange/Extensions/OpenUSD/Source/Import/Public/InterchangeUSDPipeline.h
--- a/Engine/Plugins/Interchange/Extensions/OpenUSD/Source/Import/Public/InterchangeUSDPipeline.h
+++ b/Engine/Plugins/Interchange/Extensions/OpenUSD/Source/Import/Public/InterchangeUSDPipeline.h
@@ -34,6 +34,13 @@ public:
 	UPROPERTY(EditAnywhere, BlueprintReadWrite, Category = "USD")
 	EInterchangeUsdPrimvar ImportPrimvars = EInterchangeUsdPrimvar::Bake;
 
+	/**
+	 * Subdivision level to use for all subdivision meshes on the opened stage. 0 means "don't subdivide".
+	 * The maximum level of subdivision allowed can be configured via the 'USD.Subdiv.MaxSubdivLevel' cvar.
+	 */
+	UPROPERTY(EditAnywhere, BlueprintReadWrite, Category = "USD")
+	int32 SubdivisionLevel;
+
 	/**
 	 * The translator always emits a scene node for the stage pseudoroot. Setting this option to false means we will remove it from
 	 * the factory nodes, making top-level prims lead to top-level actors and components.
diff --git a/Engine/Plugins/Interchange/Runtime/Source/Common/Public/Usd/InterchangeUsdDefinitions.h b/Engine/Plugins/Interchange/Runtime/Source/Common/Public/Usd/InterchangeUsdDefinitions.h
--- a/Engine/Plugins/Interchange/Runtime/Source/Common/Public/Usd/InterchangeUsdDefinitions.h
+++ b/Engine/Plugins/Interchange/Runtime/Source/Common/Public/Usd/InterchangeUsdDefinitions.h
@@ -48,6 +48,11 @@ namespace UE::Interchange::USD
 	// to it's field name only
 	const FString VolumeFieldNameMaterialParameterPrefix = TEXT("USD_FieldName_");
 
+	// The USD Pipeline moves its SubdivisionLevel property value onto its Mesh factory nodes under this payload
+	// attribute. The USD translator then retrieves this there and uses that subdivision level with OpenSubdiv for
+	// the actual subdivision
+	const FString SubdivisionLevelAttributeKey = TEXT("USD_SubdivisionLevel");
+
 	// Custom attribute keys used to describe SVT info extracted from the USD custom schemas.
 	// We add these to the Volume nodes that the USD translator emits.
 	namespace SparseVolumeTexture
diff --git a/Engine/Plugins/Interchange/Runtime/Source/FactoryNodes/Private/InterchangeGeometryCacheFactoryNode.cpp b/Engine/Plugins/Interchange/Runtime/Source/FactoryNodes/Private/InterchangeGeometryCacheFactoryNode.cpp
--- a/Engine/Plugins/Interchange/Runtime/Source/FactoryNodes/Private/InterchangeGeometryCacheFactoryNode.cpp
+++ b/Engine/Plugins/Interchange/Runtime/Source/FactoryNodes/Private/InterchangeGeometryCacheFactoryNode.cpp
@@ -3,6 +3,8 @@
 #include "InterchangeGeometryCacheFactoryNode.h"
 #include "Nodes/InterchangeBaseNodeContainer.h"
 
+#include "GeometryCache.h"
+
 #include UE_INLINE_GENERATED_CPP_BY_NAME(InterchangeGeometryCacheFactoryNode)
 
 UInterchangeGeometryCacheFactoryNode::UInterchangeGeometryCacheFactoryNode()
diff --git a/Engine/Plugins/Interchange/Runtime/Source/FactoryNodes/Public/InterchangeGeometryCacheFactoryNode.h b/Engine/Plugins/Interchange/Runtime/Source/FactoryNodes/Public/InterchangeGeometryCacheFactoryNode.h
--- a/Engine/Plugins/Interchange/Runtime/Source/FactoryNodes/Public/InterchangeGeometryCacheFactoryNode.h
+++ b/Engine/Plugins/Interchange/Runtime/Source/FactoryNodes/Public/InterchangeGeometryCacheFactoryNode.h
@@ -2,11 +2,12 @@
 
 #pragma once
 
-#include "GeometryCache.h"
 #include "InterchangeMeshDefinitions.h"
 #include "InterchangeMeshFactoryNode.h"
 #include "Nodes/InterchangeFactoryBaseNode.h"
 
+#include "Templates/SubclassOf.h"
+
 #include "InterchangeGeometryCacheFactoryNode.generated.h"
 
 #define UE_API INTERCHANGEFACTORYNODES_API
@@ -115,7 +116,7 @@ private:
 	UE_API virtual bool SetNodeClassFromClassAttribute() override;
 
 protected:
-	TSubclassOf<UGeometryCache> AssetClass;
+	TSubclassOf<class UGeometryCache> AssetClass;
 
 private:
 	IMPLEMENT_NODE_ATTRIBUTE_KEY(FlattenTracks);
diff --git a/Engine/Plugins/Interchange/Runtime/Source/Pipelines/Private/InterchangeGenericGeometryCachePipeline.cpp b/Engine/Plugins/Interchange/Runtime/Source/Pipelines/Private/InterchangeGenericGeometryCachePipeline.cpp
--- a/Engine/Plugins/Interchange/Runtime/Source/Pipelines/Private/InterchangeGenericGeometryCachePipeline.cpp
+++ b/Engine/Plugins/Interchange/Runtime/Source/Pipelines/Private/InterchangeGenericGeometryCachePipeline.cpp
@@ -9,6 +9,8 @@
 #include "Nodes/InterchangeBaseNode.h"
 #include "Nodes/InterchangeUserDefinedAttribute.h"
 
+#include "GeometryCache.h"
+
 void UInterchangeGenericMeshPipeline::ExecutePreImportPipelineGeometryCache()
 {
 	check(CommonMeshesProperties.IsValid());
diff --git a/Engine/Plugins/Runtime/USDCore/Source/USDUtilities/Private/USDGeomMeshConversion.cpp b/Engine/Plugins/Runtime/USDCore/Source/USDUtilities/Private/USDGeomMeshConversion.cpp
--- a/Engine/Plugins/Runtime/USDCore/Source/USDUtilities/Private/USDGeomMeshConversion.cpp
+++ b/Engine/Plugins/Runtime/USDCore/Source/USDUtilities/Private/USDGeomMeshConversion.cpp
@@ -1154,6 +1154,11 @@ namespace UE::UsdGeomMeshConversion::Private
 		pxr::VtArray<float> DisplayOpacities;
 		TArray<pxr::VtArray<pxr::GfVec2f>> UVSets;
 
+		// Attributes used for skinned meshes
+		pxr::VtArray<int> JointIndices;
+		pxr::VtArray<float> JointWeights;
+		uint32 NumInfluencesPerComponent;
+
 		// Custom attributes, that may come from primvars (used essentially for baking textures in Interchange)
 		struct FPrimvar
 		{
@@ -1178,6 +1183,8 @@ namespace UE::UsdGeomMeshConversion::Private
 		pxr::TfToken DisplayColorInterpolation;
 		pxr::TfToken DisplayOpacityInterpolation;
 		TArray<pxr::TfToken> UVSetInterpolations;
+		pxr::TfToken JointIndexInterpolation;
+		pxr::TfToken JointWeightInterpolation;
 
 		// In case those are indexed primvars, these will contain the indices
 		// Note: Velocities is not a primvar, so it can't have indices
@@ -1186,9 +1193,12 @@ namespace UE::UsdGeomMeshConversion::Private
 		pxr::VtArray<int> DisplayColorIndices;
 		pxr::VtArray<int> DisplayOpacityIndices;
 		TArray<pxr::VtArray<int>> UVSetIndices;
+		pxr::VtArray<int> JointIndexIndices;
+		pxr::VtArray<int> JointWeightIndices;
 
 		// Attributes used for subdivision
-		pxr::TfToken SubdivScheme;
+		// Default to something usable so we can subdivide e.g. UsdGeom primitives even if they don't have anything authored
+		pxr::TfToken SubdivScheme = pxr::UsdGeomTokens->catmullClark;
 		pxr::TfToken InterpolateBoundary;
 		pxr::TfToken FaceVaryingInterpolation;
 		pxr::TfToken TriangleSubdivision;
@@ -1204,9 +1214,6 @@ namespace UE::UsdGeomMeshConversion::Private
 
 		TOptional<int32> ProvidedNumUVSets;
 		int32 MaterialIndexOffset = 0;
-
-		// Attributes for SkeletalMesh:
-		TArray<FUsdRawBoneInfluence> Influences;
 	};
 
 	// https://en.cppreference.com/w/cpp/utility/variant/visit2
@@ -1375,13 +1382,45 @@ namespace UE::UsdGeomMeshConversion::Private
 		}
 	}
 
+	template<typename T>
+	void GetOrEarliest(const T& AttributeOrPrimvar, auto& OutputArray, pxr::UsdTimeCode TimeCode)
+	{
+		const bool bResult = AttributeOrPrimvar.Get(&OutputArray, TimeCode);
+
+		// Some DCCs seem to output all their mesh data in a single time sample sometimes, as opposed to on the default time code.
+		// When importing, by default we'll exclusively sample the Default timecode, which means we won't find anything in that
+		// case, and so Get() and the mesh parsing as a whole will likely fail.
+		// Here we use a fallback to also check the earliest time code in that case.
+		if (!bResult && TimeCode.IsDefault() && AttributeOrPrimvar.HasAuthoredValue())
+		{
+			if (AttributeOrPrimvar.Get(&OutputArray, pxr::UsdTimeCode::EarliestTime()))
+			{
+				pxr::UsdPrim Prim;
+				if constexpr (std::is_same_v<T, pxr::UsdAttribute>)
+				{
+					Prim = AttributeOrPrimvar.GetPrim();
+				}
+				else
+				{
+					Prim = AttributeOrPrimvar.GetAttr().GetPrim();
+				}
+
+				USD_LOG_WARNING(
+					TEXT(
+						"Mesh data attribute or primvar '%s' of prim '%s' exclusively has time samples but was sampled at the default time code, which failed. The importer will fallback to sampling the earliest time code instead."
+					),
+					*UsdToUnreal::ConvertToken(AttributeOrPrimvar.GetName()),
+					*UsdToUnreal::ConvertPath(Prim.GetPrimPath())
+				);
+			}
+		}
+	}
+
 	bool CollectMeshData(
 		const pxr::UsdPrim& Prim,
 		const UsdToUnreal::FUsdMeshConversionOptions& Options,
 		FUsdMeshData& InOutMeshData,
-		UsdUtils::FUsdPrimMaterialAssignmentInfo& InOutCombinedMaterialAssignments,
-		const pxr::UsdSkelSkinningQuery* SkinningQuery = nullptr,
-		const pxr::UsdSkelSkeletonQuery* SkeletonQuery = nullptr
+		UsdUtils::FUsdPrimMaterialAssignmentInfo& InOutCombinedMaterialAssignments
 	)
 	{
 		FScopedUsdAllocs Allocs;
@@ -1392,30 +1431,6 @@ namespace UE::UsdGeomMeshConversion::Private
 			return false;
 		}
 
-		const bool bProcessSkeletalAttributes = SkeletonQuery && SkeletonQuery->IsValid() && SkinningQuery && SkinningQuery->IsValid();
-
-		auto GetOrEarliest = [&Prim](const auto& AttributeOrPrimvar, auto& OutputArray, pxr::UsdTimeCode TimeCode)
-		{
-			const bool bResult = AttributeOrPrimvar.Get(&OutputArray, TimeCode);
-
-			// Some DCCs seem to output all their mesh data in a single time sample sometimes, as opposed to on the default time code.
-			// When importing, by default we'll exclusively sample the Default timecode, which means we won't find anything in that
-			// case, and so Get() and the mesh parsing as a whole will likely fail.
-			// Here we use a fallback to also check the earliest time code in that case.
-			if (!bResult && TimeCode.IsDefault() && AttributeOrPrimvar.HasAuthoredValue())
-			{
-				if (AttributeOrPrimvar.Get(&OutputArray, pxr::UsdTimeCode::EarliestTime()))
-				{
-					USD_LOG_WARNING(
-						TEXT(
-							"Mesh data attribute or primvar '%s' of prim '%s' exclusively has time samples but was sampled at the default time code, which failed. The importer will fallback to sampling the earliest time code instead."
-						),
-						*UsdToUnreal::ConvertToken(AttributeOrPrimvar.GetName()),
-						*UsdToUnreal::ConvertPath(Prim.GetPrimPath())
-					);
-				}
-			}
-		};
 		auto GetIndicesOrEarliest = [&Prim](const auto& AttributeOrPrimvar, auto& OutputArray, pxr::UsdTimeCode TimeCode)
 		{
 			const bool bResult = AttributeOrPrimvar.GetIndices(&OutputArray, TimeCode);
@@ -1425,29 +1440,6 @@ namespace UE::UsdGeomMeshConversion::Private
 			}
 		};
 
-		// Collect GeomBindTransform if we have one
-		bool bApplyGeomBindTransform = false;
-		pxr::GfMatrix4d GeomBindTransform{1.0};
-		pxr::GfMatrix4d InvTransposeGeomBindTransform{1.0};
-		if (bProcessSkeletalAttributes && SkinningQuery)
-		{
-			GeomBindTransform = SkinningQuery->GetGeomBindTransform(Options.TimeCode);
-			if (GeomBindTransform != pxr::GfMatrix4d(1.0))	  // This is USD's identity
-			{
-				if (GeomBindTransform.GetDeterminant() == 0.0)
-				{
-					// Can't invert, just use as-is
-					USD_LOG_WARNING(TEXT("Failed to invert geomBindTransform for prim '%s'"), *UsdToUnreal::ConvertPath(Prim.GetPrimPath()));
-					InvTransposeGeomBindTransform = GeomBindTransform;
-				}
-				else
-				{
-					InvTransposeGeomBindTransform = GeomBindTransform.GetInverse().GetTranspose();
-				}
-				bApplyGeomBindTransform = true;
-			}
-		}
-
 		// All pointsBased/Gprim attributes we'll retrieve happen to have default varying interpolation
 		const pxr::TfToken DefaultInterpolation = pxr::UsdGeomTokens->varying;
 
@@ -1499,13 +1491,6 @@ namespace UE::UsdGeomMeshConversion::Private
 				GetOrEarliest(PointsAttr, InOutMeshData.Points, Options.TimeCode);
 				InOutMeshData.PointInterpolation = GetAttrInterpolation(PointsAttr, pxr::UsdGeomTokens->vertex);
 			}
-			if (bApplyGeomBindTransform)
-			{
-				for (pxr::GfVec3f& Point : InOutMeshData.Points)
-				{
-					Point = static_cast<pxr::GfVec3f>(GeomBindTransform.Transform(Point));
-				}
-			}
 
 			// Normals
 			if (pxr::UsdGeomPrimvar NormalsPrimvar = pxr::UsdGeomPrimvar(Prim.GetAttribute(UnrealIdentifiers::PrimvarsNormals)))
@@ -1519,13 +1504,6 @@ namespace UE::UsdGeomMeshConversion::Private
 				GetOrEarliest(NormalsAttr, InOutMeshData.Normals, Options.TimeCode);
 				InOutMeshData.NormalInterpolation = GetAttrInterpolation(NormalsAttr, DefaultInterpolation);
 			}
-			if (bApplyGeomBindTransform)
-			{
-				for (pxr::GfVec3f& Normal : InOutMeshData.Normals)
-				{
-					Normal = static_cast<pxr::GfVec3f>(InvTransposeGeomBindTransform.TransformDir(Normal));
-				}
-			}
 
 			// Velocities
 			if (pxr::UsdAttribute VelocitiesAttr = UsdMesh.GetVelocitiesAttr())
@@ -1652,100 +1630,22 @@ namespace UE::UsdGeomMeshConversion::Private
 			}
 		}
 
-		// Joint Indices/Weights
-		if (bProcessSkeletalAttributes)
+		// Skeletal attributes
 		{
-			const pxr::UsdPrim& SkinningPrim = SkinningQuery->GetPrim();
-			pxr::UsdSkelBindingAPI SkelBindingAPI(SkinningPrim);
-			if (SkelBindingAPI)
+			if (pxr::UsdGeomPrimvar JointIndicesPrimvar = pxr::UsdGeomPrimvar(Prim.GetAttribute(pxr::UsdSkelTokens->primvarsSkelJointIndices)))
 			{
-				// ComputeJointInfluences returns the influences per bone that applies to all the points of the mesh
-				// ComputeVaryingJointInfluences returns the joint influences for each points, expanding the influences to all points if the mesh is
-				// rigidly deformed
-				pxr::VtArray<int> JointIndices;
-				pxr::VtArray<float> JointWeights;
-				size_t NumPoints = InOutMeshData.Points.size();
-				SkinningQuery->ComputeVaryingJointInfluences(NumPoints, &JointIndices, &JointWeights);
-
-				// Keep track of whether we added an additional "true" root bone in the cases the bound skeleton has
-				// multiple root bones
-				// We'll only ever set NumAdditionalBones to 1 or 0 (as we'll only either need a "true root bone" or
-				// not), but naming it this way allows us to use it like an offset, which should make it easier to
-				// understand whenever it is used
-				uint32 NumAdditionalBones = USDSkeletalDataConversion::HasMultipleRootBones(*SkeletonQuery) ? 1 : 0;
-				uint32 NumInfluencesPerComponent = SkinningQuery->GetNumInfluencesPerComponent();
-
-				// We keep track of which influences we added because we combine many Mesh prim (each with potentially a different
-				// explicit joint order) into the same skeletal mesh asset
-				if (JointWeights.size() > (NumPoints - 1) * (NumInfluencesPerComponent - 1))
-				{
-					uint32 JointIndex = 0;
-					InOutMeshData.Influences.Reserve(NumPoints);
-					for (uint32 PointIndex = 0; PointIndex < NumPoints; ++PointIndex)
-					{
-						// The JointIndices/JointWeights contain the influences data for NumPoints * NumInfluencesPerComponent
-						for (uint32 InfluenceIndex = 0; InfluenceIndex < NumInfluencesPerComponent; ++InfluenceIndex, ++JointIndex)
-						{
-							// BoneWeight could be 0 if the actual number of influences were less than NumInfluencesPerComponent for a given point so
-							// just ignore it
-							float BoneWeight = JointWeights[JointIndex];
-							if (BoneWeight != 0.f)
-							{
-								InOutMeshData.Influences.AddUninitialized();
-								InOutMeshData.Influences.Last().BoneIndex = NumAdditionalBones + JointIndices[JointIndex];
-								InOutMeshData.Influences.Last().Weight = BoneWeight;
-								InOutMeshData.Influences.Last().VertexIndex = PointIndex;
-							}
-						}
-					}
-				}
-				const int32 NumInfluences = InOutMeshData.Influences.Num();
-
-				// If we have a joint mapper this Mesh has an explicit joint ordering, so we need to map joint indices to the skeleton's bone indices
-				if (pxr::UsdSkelAnimMapperRefPtr AnimMapper = SkinningQuery->GetJointMapper())
-				{
-					pxr::VtArray<int> SkeletonBoneIndices;
-					if (pxr::UsdSkelSkeleton BoundSkeleton = SkelBindingAPI.GetInheritedSkeleton())
-					{
-						if (pxr::UsdAttribute SkeletonJointsAttr = BoundSkeleton.GetJointsAttr())
-						{
-							pxr::VtArray<pxr::TfToken> SkeletonJoints;
-							GetOrEarliest(SkeletonJointsAttr, SkeletonJoints, Options.TimeCode);
-
-							if (SkeletonJoints.size() > 0)
-							{
-								// If the skeleton has N bones, this will just contain { 0, 1, 2, ..., N-1 }
-								int NumUsdSkeletonBones = static_cast<int>(SkeletonJoints.size());
-								for (int SkeletonBoneIndex = 0; SkeletonBoneIndex < NumUsdSkeletonBones; ++SkeletonBoneIndex)
-								{
-									SkeletonBoneIndices.push_back(SkeletonBoneIndex);
-								}
+				GetOrEarliest(JointIndicesPrimvar, InOutMeshData.JointIndices, Options.TimeCode);
+				GetIndicesOrEarliest(JointIndicesPrimvar, InOutMeshData.JointIndices, Options.TimeCode);
+				InOutMeshData.JointIndexInterpolation = GetAttrInterpolation(JointIndicesPrimvar, pxr::UsdGeomTokens->vertex);
 
-								// Use the AnimMapper to produce the indices of the Mesh's joints within the Skeleton's list of joints.
-								// Example: Imagine skeleton had { "Root", "Root/Hip", "Root/Hip/Shoulder", "Root/Hip/Shoulder/Arm",
-								// "Root/Hip/Shoulder/Arm/Elbow"
-								// }, and so BoneIndexRemapping was { 0, 1, 2, 3, 4 }. Consider a Mesh that specifies the explicit joints {
-								// "Root/Hip/Shoulder", "Root/Hip/Shoulder/Arm" }, and so uses the indices 0 and 1 to refer to Shoulder and Arm. After
-								// the Remap call SkeletonBoneIndices will hold { 2, 3 }, as those are the indices of Shoulder and Arm within the
-								// skeleton's bones
-								pxr::VtArray<int> BoneIndexRemapping;
-								if (AnimMapper->Remap(SkeletonBoneIndices, &BoneIndexRemapping))
-								{
-									for (int32 AddedInfluenceIndex = 0; AddedInfluenceIndex < NumInfluences; ++AddedInfluenceIndex)
-									{
-										FUsdRawBoneInfluence& Influence = InOutMeshData.Influences[AddedInfluenceIndex];
+				InOutMeshData.NumInfluencesPerComponent = JointIndicesPrimvar.GetElementSize();
+			}
 
-										// We have to remove our "NumAdditionalBones" offset from the influence's bone index because that's a UE
-										// concept that the BoneIndexRemapping array doesn't really know about. After that, we have a bone index that
-										// matches the USD Skeleton joint order, then we can remap with BoneIndexRemapping and add our
-										// NumAdditionalBones back in so that it matches our USkeleton
-										Influence.BoneIndex = NumAdditionalBones + BoneIndexRemapping[Influence.BoneIndex - NumAdditionalBones];
-									}
-								}
-							}
-						}
-					}
-				}
+			if (pxr::UsdGeomPrimvar JointWeightsPrimvar = pxr::UsdGeomPrimvar(Prim.GetAttribute(pxr::UsdSkelTokens->primvarsSkelJointWeights)))
+			{
+				GetOrEarliest(JointWeightsPrimvar, InOutMeshData.JointWeights, Options.TimeCode);
+				GetIndicesOrEarliest(JointWeightsPrimvar, InOutMeshData.JointWeightIndices, Options.TimeCode);
+				InOutMeshData.JointWeightInterpolation = GetAttrInterpolation(JointWeightsPrimvar, pxr::UsdGeomTokens->vertex);
 			}
 		}
 
@@ -1766,6 +1666,134 @@ namespace UE::UsdGeomMeshConversion::Private
 		return true;
 	}
 
+	void ProcessSkinnedMeshData(
+		const pxr::UsdPrim& Prim,
+		const UsdToUnreal::FUsdMeshConversionOptions& Options,
+		FUsdMeshData& InOutMeshData,
+		const pxr::UsdSkelSkinningQuery* SkinningQuery,
+		const pxr::UsdSkelSkeletonQuery* SkeletonQuery
+	)
+	{
+		const bool bProcessSkeletalAttributes = SkeletonQuery && SkeletonQuery->IsValid() && SkinningQuery && SkinningQuery->IsValid();
+		if (!bProcessSkeletalAttributes)
+		{
+			return;
+		}
+
+		// Collect GeomBindTransform if we have one
+		bool bApplyGeomBindTransform = false;
+		pxr::GfMatrix4d GeomBindTransform{1.0};
+		pxr::GfMatrix4d InvTransposeGeomBindTransform{1.0};
+		if (bProcessSkeletalAttributes && SkinningQuery)
+		{
+			GeomBindTransform = SkinningQuery->GetGeomBindTransform(Options.TimeCode);
+			if (GeomBindTransform != pxr::GfMatrix4d(1.0))	  // This is USD's identity
+			{
+				if (GeomBindTransform.GetDeterminant() == 0.0)
+				{
+					// Can't invert, just use as-is
+					USD_LOG_WARNING(TEXT("Failed to invert geomBindTransform for prim '%s': Normals may be incorrect"), *UsdToUnreal::ConvertPath(Prim.GetPrimPath()));
+					InvTransposeGeomBindTransform = GeomBindTransform;
+				}
+				else
+				{
+					InvTransposeGeomBindTransform = GeomBindTransform.GetInverse().GetTranspose();
+				}
+				bApplyGeomBindTransform = true;
+			}
+		}
+
+		// Apply the geomBindTransform if we have a relevant one
+		if (bApplyGeomBindTransform)
+		{
+			for (pxr::GfVec3f& Point : InOutMeshData.Points)
+			{
+				Point = static_cast<pxr::GfVec3f>(GeomBindTransform.Transform(Point));
+			}
+			for (pxr::GfVec3f& Normal : InOutMeshData.Normals)
+			{
+				Normal = static_cast<pxr::GfVec3f>(InvTransposeGeomBindTransform.TransformDir(Normal));
+			}
+		}
+
+		const size_t NumPoints = InOutMeshData.Points.size();
+
+		// If our influences are marked as "constant", they won't necessarily have a value for each point, so here
+		// we "flatten" them by copy pasting that value until they match the count of points.
+		// See UsdSkelSkinningQuery::ComputeVaryingJointInfluences()
+		const bool bIsRigidlyDeformed = InOutMeshData.JointIndexInterpolation == pxr::UsdGeomTokens->constant;
+		if (bIsRigidlyDeformed)
+		{
+			pxr::UsdSkelExpandConstantInfluencesToVarying(&InOutMeshData.JointIndices, NumPoints);
+			pxr::UsdSkelExpandConstantInfluencesToVarying(&InOutMeshData.JointWeights, NumPoints);
+		}
+
+		// Compensate for any additional "true root bone" we may add in case the bound skeleton has multiple root bones
+		// We'll only ever set NumAdditionalBones to 1 or 0 (as we'll only either need a "true root bone" or
+		// not), but naming it this way allows us to use it like an offset, which should make it easier to
+		// understand whenever it is used
+		// TODO: Can this be done simpler by just performing it after the joint remapping?
+		uint32 NumAdditionalBones = USDSkeletalDataConversion::HasMultipleRootBones(*SkeletonQuery) ? 1 : 0;
+		if (NumAdditionalBones > 0)
+		{
+			for (int32 Index = 0; Index < InOutMeshData.JointIndices.size(); ++Index)
+			{
+				InOutMeshData.JointIndices[Index] += NumAdditionalBones;
+			}
+		}
+
+		// Remap the joint indices according to the Mesh's desired joint order
+		const pxr::UsdPrim& SkinningPrim = SkinningQuery->GetPrim();
+		pxr::UsdSkelBindingAPI SkelBindingAPI(SkinningPrim);
+		if (SkelBindingAPI)
+		{
+			// If we have a joint mapper this Mesh has an explicit joint ordering, so we need to map joint indices to the skeleton's bone indices
+			if (pxr::UsdSkelAnimMapperRefPtr AnimMapper = SkinningQuery->GetJointMapper())
+			{
+				pxr::VtArray<int> SkeletonBoneIndices;
+				if (pxr::UsdSkelSkeleton BoundSkeleton = SkelBindingAPI.GetInheritedSkeleton())
+				{
+					if (pxr::UsdAttribute SkeletonJointsAttr = BoundSkeleton.GetJointsAttr())
+					{
+						pxr::VtArray<pxr::TfToken> SkeletonJoints;
+						GetOrEarliest(SkeletonJointsAttr, SkeletonJoints, Options.TimeCode);
+
+						if (SkeletonJoints.size() > 0)
+						{
+							// If the skeleton has N bones, this will just contain { 0, 1, 2, ..., N-1 }
+							int NumUsdSkeletonBones = static_cast<int>(SkeletonJoints.size());
+							for (int SkeletonBoneIndex = 0; SkeletonBoneIndex < NumUsdSkeletonBones; ++SkeletonBoneIndex)
+							{
+								SkeletonBoneIndices.push_back(SkeletonBoneIndex);
+							}
+
+							// Use the AnimMapper to produce the indices of the Mesh's joints within the Skeleton's list of joints.
+							// Example: Imagine skeleton had { "Root", "Root/Hip", "Root/Hip/Shoulder", "Root/Hip/Shoulder/Arm",
+							// "Root/Hip/Shoulder/Arm/Elbow"
+							// }, and so BoneIndexRemapping was { 0, 1, 2, 3, 4 }. Consider a Mesh that specifies the explicit joints {
+							// "Root/Hip/Shoulder", "Root/Hip/Shoulder/Arm" }, and so uses the indices 0 and 1 to refer to Shoulder and Arm. After
+							// the Remap call SkeletonBoneIndices will hold { 2, 3 }, as those are the indices of Shoulder and Arm within the
+							// skeleton's bones
+							pxr::VtArray<int> BoneIndexRemapping;
+							if (AnimMapper->Remap(SkeletonBoneIndices, &BoneIndexRemapping))
+							{
+								for (int32 Index = 0; Index < InOutMeshData.JointIndices.size(); ++Index)
+								{
+									// We have to remove our "NumAdditionalBones" offset from the influence's bone index because that's a UE
+									// concept that the BoneIndexRemapping array doesn't really know about. After that, we have a bone index that
+									// matches the USD Skeleton joint order, then we can remap with BoneIndexRemapping and add our
+									// NumAdditionalBones back in so that it matches our USkeleton
+									int32 CurrentIndex = InOutMeshData.JointIndices[Index];
+									InOutMeshData.JointIndices[Index] = NumAdditionalBones + BoneIndexRemapping[CurrentIndex - NumAdditionalBones];
+								}
+							}
+						}
+					}
+				}
+			}
+		}
+	}
+
 	// OpenSubdiv expects the data elements of its buffers to implement a simple interface,
 	// so here we wrap the datatypes we'll be interpolating with that interface
 	struct FSubdivVec2f
@@ -1916,21 +1944,20 @@ namespace UE::UsdGeomMeshConversion::Private
 		Swap(FlattenedValues, SharedValuesArray);
 	};
 
-	bool SubdivideMeshData(const pxr::UsdPrim& Prim, const UsdToUnreal::FUsdMeshConversionOptions& Options, FUsdMeshData& InOutMeshData)
+	bool SubdivideMeshData(const UsdToUnreal::FUsdMeshConversionOptions& Options, FUsdMeshData& InOutMeshData)
 	{
 		// References:
 		// - USD's HdSt_OsdTopologyComputation::Resolve
 		// - UE's FSubdividePoly::ComputeTopologySubdivision
 		// - OpenSubdiv's [REDACTED_PROFILE_URL]/OpenSubdiv/blob/release/tutorials/far/tutorial_2_2/far_tutorial_2_2.cpp
 
-		FScopedUsdAllocs Allocs;
-
-		pxr::UsdGeomMesh UsdMesh{Prim};
-		if (!UsdMesh)
+		if (Options.SubdivisionLevel == 0 || InOutMeshData.SubdivScheme.IsEmpty() || InOutMeshData.SubdivScheme == pxr::UsdGeomTokens->none)
 		{
 			return false;
 		}
 
+		FScopedUsdAllocs Allocs;
+
 		int32 TargetSubdivLevel = FMath::Max(0, FMath::Min(GMaxSubdivLevel, Options.SubdivisionLevel));
 		if (TargetSubdivLevel < Options.SubdivisionLevel)
 		{
@@ -2525,6 +2552,12 @@ namespace UE::UsdGeomMeshConversion::Private
 		InOutMeshData.DisplayOpacities = ComputeFlattened(InOutMeshData.DisplayOpacities, InOutMeshData.DisplayOpacityIndices);
 		InOutMeshData.DisplayOpacityIndices = {};
 
+		InOutMeshData.JointIndices = ComputeFlattened(InOutMeshData.JointIndices, InOutMeshData.JointIndexIndices);
+		InOutMeshData.JointIndexIndices = {};
+
+		InOutMeshData.JointWeights = ComputeFlattened(InOutMeshData.JointWeights, InOutMeshData.JointWeightIndices);
+		InOutMeshData.JointWeightIndices = {};
+
 		for (int32 UVSetIndex = 0; UVSetIndex < InOutMeshData.UVSets.Num(); ++UVSetIndex)
 		{
 			pxr::VtArray<pxr::GfVec2f>& UVs = InOutMeshData.UVSets[UVSetIndex];
@@ -2723,18 +2756,19 @@ namespace UE::UsdGeomMeshConversion::Private
 		StaticMeshAttributes.Register(bKeepExistingAttribute);
 
 		// Vertex positions
+		const uint32 NumPoints = InMeshData.Points.size();
 		TVertexAttributesRef<FVector3f> MeshDescriptionVertexPositions = StaticMeshAttributes.GetVertexPositions();
 		TArray<FVertexID> VertexIDMap;
 		{
-			if (InMeshData.Points.size() < 3)
+			if (NumPoints < 3)
 			{
 				return false;
 			}
 
-			OutMeshDescription.ReserveNewVertices(InMeshData.Points.size());
-			VertexIDMap.Reserve(InMeshData.Points.size());
+			OutMeshDescription.ReserveNewVertices(NumPoints);
+			VertexIDMap.Reserve(NumPoints);
 
-			for (size_t LocalPointIndex = 0; LocalPointIndex < InMeshData.Points.size(); ++LocalPointIndex)
+			for (size_t LocalPointIndex = 0; LocalPointIndex < NumPoints; ++LocalPointIndex)
 			{
 				const pxr::GfVec3f& Point = InMeshData.Points.cdata()[LocalPointIndex];
 
@@ -3267,12 +3301,40 @@ namespace UE::UsdGeomMeshConversion::Private
 			}
 		}
 
-		if (InMeshData.Influences.Num())
+		// Handle skinned mesh data
+		if (InMeshData.JointIndices.size() > 0)
 		{
+			// Collect our joint influences in an array of structs
+			TArray<FUsdRawBoneInfluence> Influences;
+			if (InMeshData.JointWeights.size() > (NumPoints - 1) * (InMeshData.NumInfluencesPerComponent - 1))
+			{
+				uint32 JointIndex = 0;
+				Influences.Reserve(NumPoints);
+				for (uint32 PointIndex = 0; PointIndex < NumPoints; ++PointIndex)
+				{
+					// The JointIndices/JointWeights contain the influences data for NumPoints * NumInfluencesPerComponent
+					for (uint32 InfluenceIndex = 0; InfluenceIndex < InMeshData.NumInfluencesPerComponent; ++InfluenceIndex, ++JointIndex)
+					{
+						// BoneWeight could be 0 if the actual number of influences were less than NumInfluencesPerComponent for a given point so
+						// just ignore it
+						float BoneWeight = InMeshData.JointWeights[JointIndex];
+						if (BoneWeight != 0.f)
+						{
+							Influences.AddUninitialized();
+							Influences.Last().BoneIndex = InMeshData.JointIndices[JointIndex];
+							Influences.Last().Weight = BoneWeight;
+							Influences.Last().VertexIndex = PointIndex;
+						}
+					}
+				}
+			}
+
 			FSkeletalMeshAttributes MeshAttributes(OutMeshDescription);
 			MeshAttributes.Register(bKeepExistingAttribute);
+
 			FSkinWeightsVertexAttributesRef VertexSkinWeights = MeshAttributes.GetVertexSkinWeights();
-			USDSkeletalDataConversion::CopySkinWeightsToAttribute(InMeshData.Influences, VertexIDMap, VertexSkinWeights);
+
+			USDSkeletalDataConversion::CopySkinWeightsToAttribute(Influences, VertexIDMap, VertexSkinWeights);
 		}
 
 		if (NumPolygons > 0 && NumSkippedPolygons > 0)
@@ -3328,12 +3390,11 @@ bool UsdToUnreal::ConvertGeomMesh(
 
 	UsdGeomMeshImpl::FUsdMeshData MeshData;
 
-	UsdGeomMeshImpl::CollectMeshData(UsdPrim, Options, MeshData, OutMaterialAssignments, SkinningQuery, SkeletonQuery);
+	UsdGeomMeshImpl::CollectMeshData(UsdPrim, Options, MeshData, OutMaterialAssignments);
 
-	if (Options.SubdivisionLevel > 0 && MeshData.SubdivScheme != pxr::UsdGeomTokens->none)
-	{
-		UsdGeomMeshImpl::SubdivideMeshData(UsdPrim, Options, MeshData);
-	}
+	UsdGeomMeshImpl::ProcessSkinnedMeshData(UsdPrim, Options, MeshData, SkinningQuery, SkeletonQuery);
+
+	UsdGeomMeshImpl::SubdivideMeshData(Options, MeshData);
 
 	// Make sure primvars are flattened before calling ConvertMeshData.
 	// We keep faceVarying indexed primvars within CollectMeshData as they are used for subdiv,
@@ -3739,6 +3800,8 @@ bool UsdToUnreal::ConvertGeomPrimitive(
 		MeshData.PointInterpolation = pxr::UsdGeomTokens->vertex;
 	}
 
+	SubdivideMeshData(InOptions, MeshData);
+
 	FlattenIndexedPrimvars(MeshData);
 
 	pxr::UsdStageRefPtr Stage = InPrim.GetStage();
