USD: Fixed potential crash when vertex index is out of bounds in normals array during skeletal data conversion.

#rnx

diff --git a/Engine/Plugins/Runtime/USDCore/Source/USDUtilities/Private/USDSkeletalDataConversion.cpp b/Engine/Plugins/Runtime/USDCore/Source/USDUtilities/Private/USDSkeletalDataConversion.cpp
--- a/Engine/Plugins/Runtime/USDCore/Source/USDUtilities/Private/USDSkeletalDataConversion.cpp
+++ b/Engine/Plugins/Runtime/USDCore/Source/USDUtilities/Private/USDSkeletalDataConversion.cpp
@@ -1905,13 +1905,13 @@ bool UsdToUnreal::ConvertSkinnedMesh(
 
 				Triangle.WedgeIndex[FinalCornerIndex] = WedgeIndex;
 
-				if (SkelMeshImportData.bHasNormals)
+				if (SkelMeshImportData.bHasNormals && Normals.Num())
 				{
-					if (bUsingVertexInstanceNormals)
+					if (bUsingVertexInstanceNormals && Normals.IsValidIndex(OriginalVertexInstanceIndex))
 					{
 						Triangle.TangentZ[FinalCornerIndex] = Normals[OriginalVertexInstanceIndex];
 					}
-					else
+					else if (Normals.IsValidIndex(OriginalVertexIndex))
 					{
 						Triangle.TangentZ[FinalCornerIndex] = Normals[OriginalVertexIndex];
 					}
