Fixed lcalization text issue

#rnx

diff --git a/Engine/Plugins/Runtime/USDCore/Source/USDUtilities/Private/USDMaterialXShaderGraph.cpp b/Engine/Plugins/Runtime/USDCore/Source/USDUtilities/Private/USDMaterialXShaderGraph.cpp
--- a/Engine/Plugins/Runtime/USDCore/Source/USDUtilities/Private/USDMaterialXShaderGraph.cpp
+++ b/Engine/Plugins/Runtime/USDCore/Source/USDUtilities/Private/USDMaterialXShaderGraph.cpp
@@ -72,7 +72,7 @@ FUsdMaterialXShaderGraph::FUsdMaterialXShaderGraph(const pxr::UsdPrim& UsdShadeM
 
 	if (!UsdShadeMaterial)
 	{
-		USD_LOG_USERERROR(FText::Format(LOCTEXT("MaterialXShaderGraph", "Couldn't create a UsdShadeMaterial for the prim {0}"),
+		USD_LOG_USERERROR(FText::Format(LOCTEXT("MaterialXShaderGraph_NoUsdShadeMaterial", "Couldn't create a UsdShadeMaterial for the prim {0}"),
 						  FText::FromString(UsdToUnreal::ConvertToken(UsdShadeMaterialPrim.GetName()))));
 		return;
 	}
@@ -88,7 +88,7 @@ FUsdMaterialXShaderGraph::FUsdMaterialXShaderGraph(const pxr::UsdPrim& UsdShadeM
 	}
 	catch (const std::exception& Exception)
 	{
-		USD_LOG_USERERROR(FText::Format(LOCTEXT("MaterialXShaderGraph", "MaterialX: {0}\n shader graph creation aborted...\n"),
+		USD_LOG_USERERROR(FText::Format(LOCTEXT("MaterialXShaderGraph_GraphFailed", "MaterialX: {0}\n shader graph creation aborted...\n"),
 						  FText::FromString(Exception.what())));
 	}
 }
