Use FMemory_Free wrapper from MemoryBoilerPlate.h instead of plain FMemory::Free function pointer.

#rnx

diff --git a/Engine/Plugins/Runtime/USDCore/Source/UnrealUSDWrapper/Private/UnrealUSDWrapper.cpp b/Engine/Plugins/Runtime/USDCore/Source/UnrealUSDWrapper/Private/UnrealUSDWrapper.cpp
--- a/Engine/Plugins/Runtime/USDCore/Source/UnrealUSDWrapper/Private/UnrealUSDWrapper.cpp
+++ b/Engine/Plugins/Runtime/USDCore/Source/UnrealUSDWrapper/Private/UnrealUSDWrapper.cpp
@@ -1145,6 +1145,8 @@ void UnrealUSDWrapper::ClearDiagnosticDelegate()
 // Memory wrapper is declared in ModuleBoilerPlate.h
 // Needed for FMemory::Malloc to fit the signature of std::malloc
 void* FMemory_Malloc(size_t Size, size_t Alignment);
+// Needed for FMemory::Free to fit the signature of std::free
+void FMemory_Free(void* Data);
 
 class FUnrealUSDWrapperModule : public IUnrealUSDWrapperModule
 {
@@ -1155,7 +1157,7 @@ public:
 		LLM_SCOPE_BYTAG(Usd);
 
 #if !USE_USD_MEMORY_MANAGER
-		pxr::MemoryOverloads::RegisterMallocFree(FMemory_Malloc, FMemory::Free);
+		pxr::MemoryOverloads::RegisterMallocFree(FMemory_Malloc, FMemory_Free);
 #endif
 		
 		// Path to USD base plugins
