USD: Fix for LCW launch issue where it would look for USD dlls in the wrong place.

#rnx

diff --git a/Engine/Plugins/Runtime/USDCore/Source/UnrealUSDWrapper/Private/UnrealUSDWrapper.cpp b/Engine/Plugins/Runtime/USDCore/Source/UnrealUSDWrapper/Private/UnrealUSDWrapper.cpp
--- a/Engine/Plugins/Runtime/USDCore/Source/UnrealUSDWrapper/Private/UnrealUSDWrapper.cpp
+++ b/Engine/Plugins/Runtime/USDCore/Source/UnrealUSDWrapper/Private/UnrealUSDWrapper.cpp
@@ -1207,7 +1207,9 @@ public:
 #endif	  // PLATFORM_WINDOWS
 
 #else
-		FString TargetDllFolder = FPlatformProcess::BaseDir();
+		FString TargetDllFolder = FPlatformProcess::ExecutablePath();
+		FPaths::NormalizeFilename(TargetDllFolder);
+		TargetDllFolder = FPaths::GetPath(TargetDllFolder);
 #endif	  // USE_LIBRARIES_FROM_PLUGIN_FOLDER
 
 		// Path to the MaterialX standard data libraries.
