USD: Add SubdivisionAPI custom schema for per-prim subdivision level.

Update automated tests to cover per-prim subdivision level overrides;
Fix NaniteAssemblies schemas missing from Windows Arm64;
Fix the test that was meant to detect these issues but was silently broken.

diff --git a/Engine/Plugins/Interchange/Extensions/OpenUSD/Source/Import/Private/InterchangeUSDPipeline.cpp b/Engine/Plugins/Interchange/Extensions/OpenUSD/Source/Import/Private/InterchangeUSDPipeline.cpp
--- a/Engine/Plugins/Interchange/Extensions/OpenUSD/Source/Import/Private/InterchangeUSDPipeline.cpp
+++ b/Engine/Plugins/Interchange/Extensions/OpenUSD/Source/Import/Private/InterchangeUSDPipeline.cpp
@@ -991,10 +991,7 @@ namespace UE::InterchangeUsdPipeline::Private
 					UInterchangeMeshFactoryNode* MeshFactoryNode = Cast<UInterchangeMeshFactoryNode>(NodeContainer->GetFactoryNode(FactoryNodeUid));
 					if (!MeshFactoryNode)
 					{
-						if (!MeshFactoryNode)
-						{
-							return;
-						}
+						return;
 					}
 
 					if (bGenerateCompatibleMaterials)
diff --git a/Engine/Plugins/Runtime/USDCore/Resources/UsdResources/Linux/X64/plugins/unreal/resources/plugInfo.json b/Engine/Plugins/Runtime/USDCore/Resources/UsdResources/Linux/X64/plugins/unreal/resources/plugInfo.json
--- a/Engine/Plugins/Runtime/USDCore/Resources/UsdResources/Linux/X64/plugins/unreal/resources/plugInfo.json
+++ b/Engine/Plugins/Runtime/USDCore/Resources/UsdResources/Linux/X64/plugins/unreal/resources/plugInfo.json
@@ -147,6 +147,17 @@
                         ], 
                         "schemaIdentifier": "SparseVolumeTextureAPI", 
                         "schemaKind": "singleApplyAPI"
+                    }, 
+                    "unrealSubdivisionAPI": {
+                        "alias": {
+                            "UsdSchemaBase": "SubdivisionAPI"
+                        }, 
+                        "autoGenerated": true, 
+                        "bases": [
+                            "UsdAPISchemaBase"
+                        ], 
+                        "schemaIdentifier": "SubdivisionAPI", 
+                        "schemaKind": "singleApplyAPI"
                     }
                 }
             }, 
diff --git a/Engine/Plugins/Runtime/USDCore/Resources/UsdResources/Mac/plugins/unreal/resources/plugInfo.json b/Engine/Plugins/Runtime/USDCore/Resources/UsdResources/Mac/plugins/unreal/resources/plugInfo.json
--- a/Engine/Plugins/Runtime/USDCore/Resources/UsdResources/Mac/plugins/unreal/resources/plugInfo.json
+++ b/Engine/Plugins/Runtime/USDCore/Resources/UsdResources/Mac/plugins/unreal/resources/plugInfo.json
@@ -147,6 +147,17 @@
                         ], 
                         "schemaIdentifier": "SparseVolumeTextureAPI", 
                         "schemaKind": "singleApplyAPI"
+                    }, 
+                    "unrealSubdivisionAPI": {
+                        "alias": {
+                            "UsdSchemaBase": "SubdivisionAPI"
+                        }, 
+                        "autoGenerated": true, 
+                        "bases": [
+                            "UsdAPISchemaBase"
+                        ], 
+                        "schemaIdentifier": "SubdivisionAPI", 
+                        "schemaKind": "singleApplyAPI"
                     }
                 }
             }, 
diff --git a/Engine/Plugins/Runtime/USDCore/Resources/UsdResources/Win64/Arm64/plugins/unreal/resources/plugInfo.json b/Engine/Plugins/Runtime/USDCore/Resources/UsdResources/Win64/Arm64/plugins/unreal/resources/plugInfo.json
--- a/Engine/Plugins/Runtime/USDCore/Resources/UsdResources/Win64/Arm64/plugins/unreal/resources/plugInfo.json
+++ b/Engine/Plugins/Runtime/USDCore/Resources/UsdResources/Win64/Arm64/plugins/unreal/resources/plugInfo.json
@@ -88,6 +88,51 @@
                         "schemaIdentifier": "LodSubtreeAPI", 
                         "schemaKind": "singleApplyAPI"
                     }, 
+                    "unrealNaniteAssemblyExternalRefAPI": {
+                        "alias": {
+                            "UsdSchemaBase": "NaniteAssemblyExternalRefAPI"
+                        }, 
+                        "apiSchemaCanOnlyApplyTo": [
+                            "Xform"
+                        ], 
+                        "autoGenerated": true, 
+                        "bases": [
+                            "UsdAPISchemaBase"
+                        ], 
+                        "schemaIdentifier": "NaniteAssemblyExternalRefAPI", 
+                        "schemaKind": "singleApplyAPI"
+                    }, 
+                    "unrealNaniteAssemblyRootAPI": {
+                        "alias": {
+                            "UsdSchemaBase": "NaniteAssemblyRootAPI"
+                        }, 
+                        "apiSchemaCanOnlyApplyTo": [
+                            "Xform"
+                        ], 
+                        "autoGenerated": true, 
+                        "bases": [
+                            "UsdAPISchemaBase"
+                        ], 
+                        "schemaIdentifier": "NaniteAssemblyRootAPI", 
+                        "schemaKind": "singleApplyAPI"
+                    }, 
+                    "unrealNaniteAssemblySkelBindingAPI": {
+                        "alias": {
+                            "UsdSchemaBase": "NaniteAssemblySkelBindingAPI"
+                        }, 
+                        "apiSchemaCanOnlyApplyTo": [
+                            "Xform", 
+                            "Mesh", 
+                            "SkelRoot", 
+                            "PointInstancer"
+                        ], 
+                        "autoGenerated": true, 
+                        "bases": [
+                            "UsdAPISchemaBase"
+                        ], 
+                        "schemaIdentifier": "NaniteAssemblySkelBindingAPI", 
+                        "schemaKind": "singleApplyAPI"
+                    }, 
                     "unrealSparseVolumeTextureAPI": {
                         "alias": {
                             "UsdSchemaBase": "SparseVolumeTextureAPI"
@@ -102,6 +147,17 @@
                         ], 
                         "schemaIdentifier": "SparseVolumeTextureAPI", 
                         "schemaKind": "singleApplyAPI"
+                    }, 
+                    "unrealSubdivisionAPI": {
+                        "alias": {
+                            "UsdSchemaBase": "SubdivisionAPI"
+                        }, 
+                        "autoGenerated": true, 
+                        "bases": [
+                            "UsdAPISchemaBase"
+                        ], 
+                        "schemaIdentifier": "SubdivisionAPI", 
+                        "schemaKind": "singleApplyAPI"
                     }
                 }
             }, 
diff --git a/Engine/Plugins/Runtime/USDCore/Resources/UsdResources/Win64/X64/plugins/unreal/resources/plugInfo.json b/Engine/Plugins/Runtime/USDCore/Resources/UsdResources/Win64/X64/plugins/unreal/resources/plugInfo.json
--- a/Engine/Plugins/Runtime/USDCore/Resources/UsdResources/Win64/X64/plugins/unreal/resources/plugInfo.json
+++ b/Engine/Plugins/Runtime/USDCore/Resources/UsdResources/Win64/X64/plugins/unreal/resources/plugInfo.json
@@ -147,6 +147,17 @@
                         ], 
                         "schemaIdentifier": "SparseVolumeTextureAPI", 
                         "schemaKind": "singleApplyAPI"
+                    }, 
+                    "unrealSubdivisionAPI": {
+                        "alias": {
+                            "UsdSchemaBase": "SubdivisionAPI"
+                        }, 
+                        "autoGenerated": true, 
+                        "bases": [
+                            "UsdAPISchemaBase"
+                        ], 
+                        "schemaIdentifier": "SubdivisionAPI", 
+                        "schemaKind": "singleApplyAPI"
                     }
                 }
             }, 
diff --git a/Engine/Plugins/Runtime/USDCore/Source/USDUtilities/Private/USDGeomMeshConversion.cpp b/Engine/Plugins/Runtime/USDCore/Source/USDUtilities/Private/USDGeomMeshConversion.cpp
--- a/Engine/Plugins/Runtime/USDCore/Source/USDUtilities/Private/USDGeomMeshConversion.cpp
+++ b/Engine/Plugins/Runtime/USDCore/Source/USDUtilities/Private/USDGeomMeshConversion.cpp
@@ -1197,8 +1197,8 @@ namespace UE::UsdGeomMeshConversion::Private
 		pxr::VtArray<int> JointWeightIndices;
 
 		// Attributes used for subdivision
-		// Default to something usable so we can subdivide e.g. UsdGeom primitives even if they don't have anything authored
-		pxr::TfToken SubdivScheme = pxr::UsdGeomTokens->catmullClark;
+		TOptional<int32> SubdivisionLevelOverride;                    // Authored directly on the prim with the SubdivisionAPI schema
+		pxr::TfToken SubdivScheme = pxr::UsdGeomTokens->catmullClark; // Default to something usable so we can subdivide e.g. UsdGeom primitives even if they don't have anything authored
 		pxr::TfToken InterpolateBoundary;
 		pxr::TfToken FaceVaryingInterpolation;
 		pxr::TfToken TriangleSubdivision;
@@ -1463,6 +1463,17 @@ namespace UE::UsdGeomMeshConversion::Private
 			InOutMeshData.DisplayOpacityInterpolation = GetAttrInterpolation(DisplayOpacitiesPrimvar, DefaultInterpolation);
 		}
 
+		// Per-prim subdivision level (outside the UsdGeomMesh section so it also works for primitives)
+		if (Prim.HasAPI(UnrealIdentifiers::SubdivisionAPI))
+		{
+			int32 SubdivisionLevelOverride = 0;
+			pxr::UsdAttribute Attribute = Prim.GetAttribute(UnrealIdentifiers::SubdivisionLevel);
+			if (Attribute && Attribute.HasAuthoredValue() && Attribute.Get<int32>(&SubdivisionLevelOverride))
+			{
+				InOutMeshData.SubdivisionLevelOverride = SubdivisionLevelOverride;
+			}
+		}
+
 		if (pxr::UsdGeomMesh UsdMesh{Prim})
 		{
 			// Faces
@@ -1519,7 +1530,7 @@ namespace UE::UsdGeomMeshConversion::Private
 			}
 
 			// Collect the subdivision attributes only if we plan on subdividing
-			if (Options.SubdivisionLevel > 0)
+			if (InOutMeshData.SubdivisionLevelOverride.Get(Options.SubdivisionLevel) > 0)
 			{
 				if (pxr::UsdAttribute SubdivSchemeAttr = UsdMesh.GetSubdivisionSchemeAttr())
 				{
@@ -1951,15 +1962,16 @@ namespace UE::UsdGeomMeshConversion::Private
 		// - UE's FSubdividePoly::ComputeTopologySubdivision
 		// - OpenSubdiv's [REDACTED_PROFILE_URL]/OpenSubdiv/blob/release/tutorials/far/tutorial_2_2/far_tutorial_2_2.cpp
 
-		if (Options.SubdivisionLevel == 0 || InOutMeshData.SubdivScheme.IsEmpty() || InOutMeshData.SubdivScheme == pxr::UsdGeomTokens->none)
+		const int32 SubdivisionLevel = InOutMeshData.SubdivisionLevelOverride.Get(Options.SubdivisionLevel);
+		if (SubdivisionLevel == 0 || InOutMeshData.SubdivScheme.IsEmpty() || InOutMeshData.SubdivScheme == pxr::UsdGeomTokens->none)
 		{
 			return false;
 		}
 
 		FScopedUsdAllocs Allocs;
 
-		int32 TargetSubdivLevel = FMath::Max(0, FMath::Min(GMaxSubdivLevel, Options.SubdivisionLevel));
-		if (TargetSubdivLevel < Options.SubdivisionLevel)
+		int32 TargetSubdivLevel = FMath::Max(0, FMath::Min(GMaxSubdivLevel, SubdivisionLevel));
+		if (TargetSubdivLevel < SubdivisionLevel)
 		{
 			USD_LOG_USERWARNING(FText::Format(
 				LOCTEXT("ClampedSubdiv", "Max subdivision level was clamped to {0} (controlled by the cvar '{1}')"),
@@ -3287,7 +3299,7 @@ namespace UE::UsdGeomMeshConversion::Private
 				// If we're applying subdivision, pxOsd will take care of
 				// converting the orientation from leftHanded to rightHanded,
 				// so only flip if we are NOT subdividing.
-				if (InMeshData.Orientation == pxr::UsdGeomTokens->leftHanded && InOptions.SubdivisionLevel < 1)
+				if (InMeshData.Orientation == pxr::UsdGeomTokens->leftHanded && InMeshData.SubdivisionLevelOverride.Get(InOptions.SubdivisionLevel) < 1)
 				{
 					for (int32 i = 0; i < CornerInstanceIDs.Num() / 2; ++i)
 					{
diff --git a/Engine/Plugins/Runtime/USDCore/Source/UnrealUSDWrapper/Private/UnrealUSDWrapper.cpp b/Engine/Plugins/Runtime/USDCore/Source/UnrealUSDWrapper/Private/UnrealUSDWrapper.cpp
--- a/Engine/Plugins/Runtime/USDCore/Source/UnrealUSDWrapper/Private/UnrealUSDWrapper.cpp
+++ b/Engine/Plugins/Runtime/USDCore/Source/UnrealUSDWrapper/Private/UnrealUSDWrapper.cpp
@@ -96,6 +96,9 @@ namespace UnrealIdentifiers
 	const TfToken LOD("LOD");
 	const TfToken LodSubtreeAPI("LodSubtreeAPI");
 	const TfToken UnrealLodSubtreeLevels = TfToken("unreal:lodSubtree:levels");
+	
+	const TfToken SubdivisionAPI("SubdivisionAPI");
+	const TfToken SubdivisionLevel = TfToken("unreal:subdivisionLevel");
 
 	const TfToken MaterialAssignment = TfToken("unrealMaterial");
 	const TfToken Unreal = TfToken("unreal");
diff --git a/Engine/Plugins/Runtime/USDCore/Source/UnrealUSDWrapper/Public/UnrealUSDWrapper.h b/Engine/Plugins/Runtime/USDCore/Source/UnrealUSDWrapper/Public/UnrealUSDWrapper.h
--- a/Engine/Plugins/Runtime/USDCore/Source/UnrealUSDWrapper/Public/UnrealUSDWrapper.h
+++ b/Engine/Plugins/Runtime/USDCore/Source/UnrealUSDWrapper/Public/UnrealUSDWrapper.h
@@ -477,6 +477,9 @@ namespace UnrealIdentifiers
 	extern UNREALUSDWRAPPER_API const pxr::TfToken LodSubtreeAPI;
 	extern UNREALUSDWRAPPER_API const pxr::TfToken UnrealLodSubtreeLevels;
 
+	extern UNREALUSDWRAPPER_API const pxr::TfToken SubdivisionAPI;
+	extern UNREALUSDWRAPPER_API const pxr::TfToken SubdivisionLevel;
+
 	/* Attribute name when assigning Unreal materials to UsdGeomMeshes */
 	extern UNREALUSDWRAPPER_API const pxr::TfToken MaterialAssignment;
 	extern UNREALUSDWRAPPER_API const pxr::TfToken Unreal;
