USD MaterialX:
- add support for nodegraph when creating the MaterialX shader graph

diff --git a/Engine/Plugins/Runtime/USDCore/Source/USDUtilities/Private/USDMaterialXShaderGraph.cpp b/Engine/Plugins/Runtime/USDCore/Source/USDUtilities/Private/USDMaterialXShaderGraph.cpp
--- a/Engine/Plugins/Runtime/USDCore/Source/USDUtilities/Private/USDMaterialXShaderGraph.cpp
+++ b/Engine/Plugins/Runtime/USDCore/Source/USDUtilities/Private/USDMaterialXShaderGraph.cpp
@@ -50,6 +50,7 @@ FUsdMaterialXShaderGraph::FUsdMaterialXShaderGraph(const pxr::UsdPrim& UsdShadeM
 		{"color3f",  "color3"},
 		{"color4f",  "color4"},
 		{"float2",   "vector2"},
+		{"float3",   "vector3"},
 		{"vector3f", "vector3"},
 		{"vector4f", "vector4"},
 		{"int",      "integer"},
@@ -140,7 +141,7 @@ void FUsdMaterialXShaderGraph::CreateMaterial(const pxr::UsdShadeMaterial& UsdSh
 		mx::NodePtr Shader = InDocument->addNodeInstance(NodeDef, ShadeShader.GetPrim().GetName().GetString());
 
 		// if the shadergraph is not valid (probably caused by a missing nodedef or wrong type), just invalidate it all
-		if (!ComputeShaderGraph(Shader, ShadeShader))
+		if (!ComputeShaderGraph(Shader, ShadeShader, Document))
 		{
 			Shader = nullptr;
 		}
@@ -149,25 +150,14 @@ void FUsdMaterialXShaderGraph::CreateMaterial(const pxr::UsdShadeMaterial& UsdSh
 	};
 
 	// Create and store the input interface names
-	TUsdStore<std::vector<pxr::UsdShadeInput>> InterfaceInputs = UsdShadeMaterial.GetInterfaceInputs();
-
-	for (const pxr::UsdShadeInput& InterfaceInput : *InterfaceInputs)
-	{
-		std::string InterfaceNameInput = InterfaceInput.GetBaseName().GetString();
-		std::string InterfaceNameType = InterfaceInput.GetTypeName().GetAsToken().GetString();
-		GetMatchingInputType(InterfaceNameInput, InterfaceNameType);
-
-		mx::InputPtr InputInterfaceName = Document->addInput(InterfaceNameInput, InterfaceNameType);
-		SetInputValue(InputInterfaceName, InterfaceInput);
-		InterfaceNames.insert(InputInterfaceName->getName());
-	}
+	GetInterfaceNameInputs(UsdShadeMaterial, Document);
 
 	// Create the shader graph for the surface/volume/displacement shader
 	// for the following source materials since we create them after the shader graph, we need to make sure that their name is unique
 	// unlike MaterialX, USD allows a same name between a parent and its child
 	std::string MaterialName = UsdShadeMaterial.GetPrim().GetName().GetString();
 	mx::NodePtr SurfaceMaterial;
-	if(SurfaceSource)
+	if (SurfaceSource)
 	{
 		mx::NodePtr SurfaceShader = CreateShaderGraph(SurfaceSource, Document);
 		if (SurfaceShader)
@@ -196,7 +186,7 @@ void FUsdMaterialXShaderGraph::CreateMaterial(const pxr::UsdShadeMaterial& UsdSh
 		}
 	}
 
-	if(SurfaceMaterial)
+	if (SurfaceMaterial)
 	{
 		std::string ErrorMessage;
 		SurfaceMaterial->validate(&ErrorMessage);
@@ -210,7 +200,7 @@ void FUsdMaterialXShaderGraph::CreateMaterial(const pxr::UsdShadeMaterial& UsdSh
 
 	// a volume material is another type of material different from a surface material, it's unclear if we can have both a surface and volume in the same Material
 	// if that's the case let's just suffix the name of the volume material to avoid collision during the creation of the node
-	if(VolumeSource)
+	if (VolumeSource)
 	{
 		if (mx::NodePtr VolumeShader = CreateShaderGraph(VolumeSource, Document))
 		{
@@ -233,7 +223,7 @@ void FUsdMaterialXShaderGraph::CreateMaterial(const pxr::UsdShadeMaterial& UsdSh
 
 }
 
-bool FUsdMaterialXShaderGraph::ComputeShaderGraph(mx::NodePtr Node, const pxr::UsdShadeShader& Shader)
+bool FUsdMaterialXShaderGraph::ComputeShaderGraph(mx::NodePtr Node, const pxr::UsdShadeShader& Shader, mx::GraphElementPtr Graph)
 {
 	if (!Node)
 	{
@@ -263,14 +253,23 @@ bool FUsdMaterialXShaderGraph::ComputeShaderGraph(mx::NodePtr Node, const pxr::U
 		// Let's loop over the connected sources, even though we should expect only one
 		for (const pxr::UsdShadeConnectionSourceInfo& Source : *ConnectedSources)
 		{
-			if (SetInterfaceName(Input, Source))
+			// we'll process the rest of the shader graph inside the nodegraph
+			if (mx::NodeGraphPtr NodeGraph = GetNodeGraph(Source, Graph))
+			{
+				Input->setNodeGraphString(NodeGraph->getName());
+				Input->setOutputString(Source.sourceName);
+				bHasValue = false;
+				continue;
+			}
+
+			if (SetInterfaceName(Input, Source, Graph))
 			{
 				bHasValue = false;
 				break;
 			}
 
 			bool bIsNodeAlreadyExist;
-			mx::NodePtr ConnectedNode = GetNode(Source, bIsNodeAlreadyExist);
+			mx::NodePtr ConnectedNode = GetNode(Source, Graph, bIsNodeAlreadyExist);
 			// Let's set if the input has an interfaceName
 			if (!ConnectedNode)
 			{
@@ -278,7 +277,7 @@ bool FUsdMaterialXShaderGraph::ComputeShaderGraph(mx::NodePtr Node, const pxr::U
 			}
 
 			Input->setConnectedNode(ConnectedNode);
-			
+
 			// no need to add the default output
 			if (const std::string& OutputName = Source.sourceName.GetString(); OutputName != "out")
 			{
@@ -288,7 +287,7 @@ bool FUsdMaterialXShaderGraph::ComputeShaderGraph(mx::NodePtr Node, const pxr::U
 			// recurse if it's a newly created node, otherwise the work has already been done
 			if (!bIsNodeAlreadyExist)
 			{
-				bHasValue = ComputeShaderGraph(ConnectedNode, Source.source);
+				bHasValue = ComputeShaderGraph(ConnectedNode, Source.source, Graph);
 			}
 
 			bHasValue = false;
@@ -328,8 +327,9 @@ void FUsdMaterialXShaderGraph::SetInputValue(mx::InputPtr Input, const pxr::UsdS
 
 		if (FPairingUsdMtlxType Value; UsdInput.Get<FPairingUsdMtlxType>(&Value))
 		{
-			if constexpr (std::is_same_v<FPairingUsdMtlxType, pxr::GfVec4f> || std::is_same_v<FPairingUsdMtlxType, pxr::GfVec3f>
-						  || std::is_same_v<FPairingUsdMtlxType, pxr::GfVec2f>)
+			if constexpr (std::is_same_v<FPairingUsdMtlxType, pxr::GfVec4f> ||
+						  std::is_same_v<FPairingUsdMtlxType, pxr::GfVec3f> ||
+						  std::is_same_v<FPairingUsdMtlxType, pxr::GfVec2f>)
 			{
 				Input->setValue(MtlxType{Value.data(), Value.data() + FPairingUsdMtlxType::dimension});
 			}
@@ -422,22 +422,27 @@ void FUsdMaterialXShaderGraph::SetInputValue(mx::InputPtr Input, const pxr::UsdS
 	}
 }
 
-bool FUsdMaterialXShaderGraph::SetInterfaceName(MaterialX::InputPtr Input, const pxr::UsdShadeConnectionSourceInfo& Source)
+bool FUsdMaterialXShaderGraph::SetInterfaceName(mx::InputPtr Input, const pxr::UsdShadeConnectionSourceInfo& Source, mx::GraphElementPtr Graph)
 {
 	bool bIsSet = false;
 	const std::string& SourceName = Source.sourceName.GetString();
-	if (InterfaceNames.find(SourceName) != InterfaceNames.end())
+
+	// the test might not be needed since we always add at least the Document inside the InterfaceNames
+	if (auto It = InterfaceNames.find(Graph); It != InterfaceNames.end())
 	{
-		if (mx::InputPtr InputInterfaceName = Document->getInput(SourceName))
+		if (It->second.find(SourceName) != It->second.end())
 		{
-			Input->setInterfaceName(InputInterfaceName->getName());
-			bIsSet = true;
+			if (mx::InputPtr InputInterfaceName = It->first->getInput(SourceName))
+			{
+				Input->setInterfaceName(InputInterfaceName->getName());
+				bIsSet = true;
+			}
 		}
 	}
 	return bIsSet;
 }
 
-mx::NodePtr FUsdMaterialXShaderGraph::GetNode(const pxr::UsdShadeConnectionSourceInfo& Source, bool& bIsNodeAlreadyExist)
+mx::NodePtr FUsdMaterialXShaderGraph::GetNode(const pxr::UsdShadeConnectionSourceInfo& Source, mx::GraphElementPtr Graph, bool& bIsNodeAlreadyExist)
 {
 	bool bIsGeomPropValue;
 	bIsNodeAlreadyExist = true;
@@ -450,12 +455,12 @@ mx::NodePtr FUsdMaterialXShaderGraph::GetNode(const pxr::UsdShadeConnectionSourc
 
 	std::string ConnectedShaderName = Source.source.GetPrim().GetName().GetString();
 
-	mx::NodePtr ConnectedNode = Document->getNode(ConnectedShaderName);
+	mx::NodePtr ConnectedNode = Graph->getNode(ConnectedShaderName);
 
 	// if the node is not present in the graph, we create it and notify that it's a new node
 	if (!ConnectedNode)
 	{
-		ConnectedNode = Document->addNodeInstance(NodeDef, ConnectedShaderName);
+		ConnectedNode = Graph->addNodeInstance(NodeDef, ConnectedShaderName);
 		if (bIsGeomPropValue)
 		{
 			ConnectedNode->setTypedAttribute("UE:GeomPropImage", bIsGeomPropValue);
@@ -471,7 +476,7 @@ mx::InputPtr FUsdMaterialXShaderGraph::GetInput(mx::NodePtr Node, const pxr::Usd
 	std::string InputName = UsdInput.GetBaseName().GetString();
 	std::string InputType = UsdInput.GetTypeName().GetAsToken().GetString();
 
-	bool bMatch = GetMatchingInputType(InputName, InputType);
+	bool bMatch = GetMatchingType(InputName, InputType);
 
 	if (!bMatch && InputName == "geomprop")
 	{
@@ -510,7 +515,7 @@ mx::InputPtr FUsdMaterialXShaderGraph::GetInput(mx::NodePtr Node, const pxr::Usd
 
 mx::NodeDefPtr FUsdMaterialXShaderGraph::GetNodeDef(const pxr::UsdShadeConnectionSourceInfo& Source, bool& bIsGeomPropValue)
 {
-	pxr::UsdShadeShader UsdConnectedShader{Source.source.GetPrim()};
+	pxr::UsdShadeShader UsdConnectedShader{ Source.source.GetPrim() };
 
 	if (!UsdConnectedShader)
 	{
@@ -570,20 +575,73 @@ mx::NodeDefPtr FUsdMaterialXShaderGraph::GetNodeDef(const pxr::UsdShadeConnectio
 	return NodeDef;
 }
 
-bool FUsdMaterialXShaderGraph::GetMatchingInputType(const std::string& InputName, std::string& InOutInputType)
+mx::NodeGraphPtr FUsdMaterialXShaderGraph::GetNodeGraph(const pxr::UsdShadeConnectionSourceInfo& Source, mx::GraphElementPtr Graph)
+{
+	// we only process the nodegraph, if it's connected to a UsdShadeMaterial then we'll process it as a nodename or a value
+	if (pxr::UsdShadeMaterial{ Source.source.GetPrim() })
+	{
+		return nullptr;
+	}
+
+	pxr::UsdShadeNodeGraph UsdConnectedNodeGraph{ Source.source.GetPrim() };
+
+	if (!UsdConnectedNodeGraph)
+	{
+		return nullptr;
+	}
+
+	std::string NodeGraphName = UsdConnectedNodeGraph.GetPrim().GetName().GetString();
+	mx::NodeGraphPtr NodeGraph = Document->getNodeGraph(NodeGraphName);
+
+	if (NodeGraph)
+	{
+		return NodeGraph;
+	}
+
+	NodeGraph = Document->addNodeGraph(NodeGraphName);
+	GetInterfaceNameInputs(UsdConnectedNodeGraph, NodeGraph);
+
+	std::vector<pxr::UsdShadeOutput> UsdShadeOutputs = UsdConnectedNodeGraph.GetOutputs();
+	for (const pxr::UsdShadeOutput& UsdShadeOutput : UsdShadeOutputs)
+	{
+		std::string OutputName = UsdShadeOutput.GetBaseName().GetString();
+		std::string OutputType = UsdShadeOutput.GetTypeName().GetAsToken().GetString();
+
+		GetMatchingType(OutputName, OutputType);
+		mx::OutputPtr Output = NodeGraph->addOutput(OutputName, OutputType);
+
+		TUsdStore<pxr::UsdShadeOutput::SourceInfoVector> ConnectedSources = UsdShadeOutput.GetConnectedSources();
+
+		for (const pxr::UsdShadeConnectionSourceInfo& ConnectedSource : *ConnectedSources)
+		{
+			bool bIsNodeAlreadyExist;
+			mx::NodePtr ConnectedNode = GetNode(ConnectedSource, NodeGraph, bIsNodeAlreadyExist);
+			if (!ConnectedNode)
+			{
+				break;
+			}
+			Output->setNodeName(ConnectedNode->getName());
+			ComputeShaderGraph(ConnectedNode, ConnectedSource.source, NodeGraph);
+		}
+	}
+
+	return NodeGraph;
+}
+
+bool FUsdMaterialXShaderGraph::GetMatchingType(const std::string& Name, std::string& InOutType)
 {
 	bool bMatch = false;
 
-	if (MaterialXTypes.find(InOutInputType) == MaterialXTypes.end())
+	if (MaterialXTypes.find(InOutType) == MaterialXTypes.end())
 	{
-		if (auto It = UsdToMaterialXTypes.find(InOutInputType); It != UsdToMaterialXTypes.end())
+		if (auto It = UsdToMaterialXTypes.find(InOutType); It != UsdToMaterialXTypes.end())
 		{
-			InOutInputType = It->second;
+			InOutType = It->second;
 			bMatch = true;
 		}
-		else if (InputName == "file")	 // we're dealing with a texture, USD calls the type "asset" for the input
+		else if (Name == "file")	 // we're dealing with a texture, USD calls the type "asset" for the input
 		{
-			InOutInputType = "filename";
+			InOutType = "filename";
 			bMatch = true;
 		}
 	}
@@ -596,5 +654,22 @@ bool FUsdMaterialXShaderGraph::IsTangentSpaceInput(const pxr::UsdShadeInput& Usd
 	return TangentSpaceInputs.find(UsdInput.GetBaseName().GetString()) != TangentSpaceInputs.cend();
 }
 
+void FUsdMaterialXShaderGraph::GetInterfaceNameInputs(const pxr::UsdShadeNodeGraph& UsdShadeGraph, mx::GraphElementPtr Graph)
+{
+	TUsdStore<std::vector<pxr::UsdShadeInput>> InterfaceInputs = UsdShadeGraph.GetInterfaceInputs();
+	std::unordered_set<std::string> InterfaceInputNames;
+	for (const pxr::UsdShadeInput& InterfaceInput : *InterfaceInputs)
+	{
+		std::string InterfaceNameInput = InterfaceInput.GetBaseName().GetString();
+		std::string InterfaceNameType = InterfaceInput.GetTypeName().GetAsToken().GetString();
+		GetMatchingType(InterfaceNameInput, InterfaceNameType);
+
+		mx::InputPtr InputInterfaceName = Graph->addInput(InterfaceNameInput, InterfaceNameType);
+		SetInputValue(InputInterfaceName, InterfaceInput);
+		InterfaceInputNames.emplace(InputInterfaceName->getName());
+	}
+	InterfaceNames.emplace(Graph, InterfaceInputNames);
+}
+
 #undef LOCTEXT_NAMESPACE 
 #endif // USE_USD_SDK && ENABLE_USD_MATERIALX
diff --git a/Engine/Plugins/Runtime/USDCore/Source/USDUtilities/Public/USDMaterialXShaderGraph.h b/Engine/Plugins/Runtime/USDCore/Source/USDUtilities/Public/USDMaterialXShaderGraph.h
--- a/Engine/Plugins/Runtime/USDCore/Source/USDUtilities/Public/USDMaterialXShaderGraph.h
+++ b/Engine/Plugins/Runtime/USDCore/Source/USDUtilities/Public/USDMaterialXShaderGraph.h
@@ -14,7 +14,7 @@
 #include "MaterialXFormat/Util.h"
 #endif // USE_USD_SDK && ENABLE_USD_MATERIALX
 
-/** 
+/**
  * Create a MaterialX shader graph from a UsdPrim, all <geompropvalue> nodes are either converted to an <image> or a proper <geomcolor>
  * Just get the Document in order to search for the surface material, by just looking for the name of the Prim it will return the correct surface material.
  */
@@ -30,7 +30,7 @@ public:
 
 #if USE_USD_SDK && ENABLE_USD_MATERIALX
 
-	/**
+	/** 
 	 * @param UsdShadeMaterialPrim - The prim used to generate the shader graph, should be a valid pxr::UsdShadeMaterial
 	 * @param RenderContext - Which render context output to read from the UsdShadeMaterial
 	 */
@@ -46,37 +46,41 @@ private:
 	/** Create either a <surfacematerial> or a <volumematerial> if possible. */
 	void CreateMaterial(const pxr::UsdShadeMaterial& UsdShadeMaterial, const TCHAR* RenderContext);
 
-	bool ComputeShaderGraph(MaterialX::NodePtr Node, const pxr::UsdShadeShader& Shader);
+	bool ComputeShaderGraph(MaterialX::NodePtr Node, const pxr::UsdShadeShader& Shader, MaterialX::GraphElementPtr Graph);
 
 	bool IsGeomColor(const pxr::UsdShadeShader& GeomPropValueShader);
 
 	void SetInputValue(MaterialX::InputPtr Input, const pxr::UsdShadeInput& UsdInput);
 
-	/** Set the interfaceName of the input if any */ 
-	bool SetInterfaceName(MaterialX::InputPtr Input, const pxr::UsdShadeConnectionSourceInfo& Source);
+	/** Set the interfaceName of the input if any */
+	bool SetInterfaceName(MaterialX::InputPtr Input, const pxr::UsdShadeConnectionSourceInfo& Source, MaterialX::GraphElementPtr Graph);
 
 	/**
 	 * Get the MaterialX node from a UsdShadeConnectionSourceInfo
 	 * @param Source - upstream source in the Shader Graph, usually the connected node to an input
+	 * @param Graph - either the Document or a nodegraph
 	 * @param bIsNodeAlreadyExist - return true if the node was already present in the shader graph, false otherwise
 	 * @return The corresponding MaterialX node in the shader graph, if it doesn't exist it creates one
 	 */
-	MaterialX::NodePtr GetNode(const pxr::UsdShadeConnectionSourceInfo& Source, bool& bIsNodeAlreadyExist);
+	MaterialX::NodePtr GetNode(const pxr::UsdShadeConnectionSourceInfo& Source, MaterialX::GraphElementPtr Graph, bool& bIsNodeAlreadyExist);
 
 	MaterialX::InputPtr GetInput(MaterialX::NodePtr Node, const pxr::UsdShadeInput& UsdInput);
 
 	MaterialX::NodeDefPtr GetNodeDef(const pxr::UsdShadeConnectionSourceInfo& Source, bool& bIsGeomPropValue);
 
+	MaterialX::NodeGraphPtr GetNodeGraph(const pxr::UsdShadeConnectionSourceInfo& Source, MaterialX::GraphElementPtr Graph);
 	/**
-	 * Given an input type, retrieve the MaterialX associated type
-	 * @param InputName - The name of the input
-	 * @param InOutInputType - The type of the input, if it's valid it will now be the MaterialX type
+	 * Given an input/output type, retrieve the MaterialX associated type
+	 * @param Name - The name of the input/output
+	 * @param InOutType - The type of the input/output, if it's valid it will now be the MaterialX type
 	 * @return true if there is a valid match between USD and MaterialX or if it's already a MaterialX type, false otherwise
 	 */
-	bool GetMatchingInputType(const std::string& InputName, std::string& InOutInputType);
+	bool GetMatchingType(const std::string& Name, std::string& InOutType);
 
 	bool IsTangentSpaceInput(const pxr::UsdShadeInput& UsdInput) const;
 
+	void GetInterfaceNameInputs(const pxr::UsdShadeNodeGraph& UsdShadeGraph, MaterialX::GraphElementPtr Graph);
+
 private:
 
 	std::unordered_set<std::string> MaterialXTypes;
@@ -86,8 +90,8 @@ private:
 	// Names of the inputs from a <surfaceshader> that requires to be in tangent space for the baking of the primvars later on
 	std::unordered_set<std::string> TangentSpaceInputs;
 
-	// MaterialX interfacenames (USD calls that InterfaceInputs), some inputs use them instead of a value or a direct connection to a node
-	std::unordered_set<std::string> InterfaceNames; 
+	// MaterialX interfacenames relatives to their nodegraph (the Document if no nodegraph), some inputs use them instead of a value or a direct connection to a node
+	std::unordered_map<MaterialX::GraphElementPtr, std::unordered_set<std::string>> InterfaceNames;
 
 	MaterialX::DocumentPtr Document;
 
