USDCore, Interchange: forward declare USD types to address RTTI-related compilation issues

The Interchange modules do not enable RTTI, but some of them
do make use of OpenUSD/pxr API, some of which requires RTTI.
One example is the FUSDMaterialXShaderGraph class provided by
the USDUtilities module (which enables RTTI), which is a dependency
of InterchangeOpenUSDImport (which does *not* enable RTTI).
In particular, this is a problem on Linux where with RTTI disabled
for the InterchangeOpenUSDImport module (as it currently is), the
compiler would complain about missing type information for USD
types. Turning RTTI on for InterchangeOpenUSDImport is not currently
an option because it includes types derived from engine classes, so
with RTTI enabled for InterchangeOpenUSDImport, the compiler would
complain about missing type information for UE types (like UObject,
UInterchangePipelineBase, etc.).

We can avoid these issues and leave RTTI disabled in
InterchangeOpenUSDImport by ensuring that pxr types are only
forward declared in FUSDMaterialXShaderGraph and that otherwise
no pxr headers that might require RTTI are included.

The USDUtilities and InterchangeImport modules also have
public header files that #include MaterialX headers, so MaterialX
is declared as a public dependency for those modules.

diff --git a/Engine/Plugins/Interchange/Runtime/Source/Import/InterchangeImport.Build.cs b/Engine/Plugins/Interchange/Runtime/Source/Import/InterchangeImport.Build.cs
--- a/Engine/Plugins/Interchange/Runtime/Source/Import/InterchangeImport.Build.cs
+++ b/Engine/Plugins/Interchange/Runtime/Source/Import/InterchangeImport.Build.cs
@@ -55,6 +55,15 @@ namespace UnrealBuildTool.Rules
 				bEnableExceptions = true;
 				bDisableAutoRTFMInstrumentation = true;
 
+				// Public because MaterialX headers are #include'd in
+				// public InterchangeImport headers.
+				PublicDependencyModuleNames.AddRange(
+					new string[]
+					{
+						"MaterialX"
+					}
+				);
+
 				PrivateDependencyModuleNames.AddRange(
 					new string[]
 					{
@@ -62,7 +71,6 @@ namespace UnrealBuildTool.Rules
 						"HairStrandsEditor",
 						"InterchangeFbxParser",
 						"MaterialEditor",
-						"MaterialX",
 						"SkeletalMeshUtilitiesCommon",
 						"UnrealEd",
 						"VariantManager",
diff --git a/Engine/Plugins/Runtime/USDCore/Source/USDUtilities/Private/USDMaterialXShaderGraph.cpp b/Engine/Plugins/Runtime/USDCore/Source/USDUtilities/Private/USDMaterialXShaderGraph.cpp
--- a/Engine/Plugins/Runtime/USDCore/Source/USDUtilities/Private/USDMaterialXShaderGraph.cpp
+++ b/Engine/Plugins/Runtime/USDCore/Source/USDUtilities/Private/USDMaterialXShaderGraph.cpp
@@ -11,6 +11,17 @@
 #include "USDTypesConversion.h"
 #include "UsdWrappers/SdfLayer.h"
 
+#include "USDIncludesStart.h"
+#include "pxr/base/tf/token.h"
+#include "pxr/usd/sdf/layerUtils.h"
+#include "pxr/usd/usd/prim.h"
+#include "pxr/usd/usdShade/connectableAPI.h"
+#include "pxr/usd/usdShade/input.h"
+#include "pxr/usd/usdShade/material.h"
+#include "pxr/usd/usdShade/nodeGraph.h"
+#include "pxr/usd/usdShade/shader.h"
+#include "USDIncludesEnd.h"
+
 #define LOCTEXT_NAMESPACE "USDMaterialXShaderGraph"
 
 namespace mx = MaterialX;
diff --git a/Engine/Plugins/Runtime/USDCore/Source/USDUtilities/Public/USDMaterialXShaderGraph.h b/Engine/Plugins/Runtime/USDCore/Source/USDUtilities/Public/USDMaterialXShaderGraph.h
--- a/Engine/Plugins/Runtime/USDCore/Source/USDUtilities/Public/USDMaterialXShaderGraph.h
+++ b/Engine/Plugins/Runtime/USDCore/Source/USDUtilities/Public/USDMaterialXShaderGraph.h
@@ -6,12 +6,26 @@
 
 #if USE_USD_SDK && ENABLE_USD_MATERIALX
 #include "USDIncludesStart.h"
-#include "pxr/base/tf/token.h"
-#include "pxr/usd/sdf/layerUtils.h"
-#include "pxr/usd/usdShade/material.h"
+#include "pxr/pxr.h"
+#include "pxr/base/gf/vec2f.h"
+#include "pxr/base/gf/vec3f.h"
+#include "pxr/base/gf/vec4f.h"
 #include "USDIncludesEnd.h"
 
 #include "MaterialXFormat/Util.h"
+
+#include <unordered_map>
+#include <unordered_set>
+
+PXR_NAMESPACE_OPEN_SCOPE
+	class UsdPrim;
+	struct UsdShadeConnectionSourceInfo;
+	class UsdShadeInput;
+	class UsdShadeMaterial;
+	class UsdShadeNodeGraph;
+	class UsdShadeShader;
+PXR_NAMESPACE_CLOSE_SCOPE
+
 #endif // USE_USD_SDK && ENABLE_USD_MATERIALX
 
 /**
diff --git a/Engine/Plugins/Runtime/USDCore/Source/USDUtilities/USDUtilities.Build.cs b/Engine/Plugins/Runtime/USDCore/Source/USDUtilities/USDUtilities.Build.cs
--- a/Engine/Plugins/Runtime/USDCore/Source/USDUtilities/USDUtilities.Build.cs
+++ b/Engine/Plugins/Runtime/USDCore/Source/USDUtilities/USDUtilities.Build.cs
@@ -71,7 +71,9 @@ namespace UnrealBuildTool.Rules
 				// we work through issues affecting Mac and Linux.
 				if (Target.Platform == UnrealTargetPlatform.Win64)
 				{
-					PrivateDependencyModuleNames.AddRange(
+					// Public because MaterialX headers are #include'd in
+					// public USDUtilities headers.
+					PublicDependencyModuleNames.AddRange(
 						new string[]
 						{
 							"MaterialX"
