USDCore, USDImporter: re-enable MaterialX via USD on Linux

Note that the "test_materialx_standard_surface" and "test_material_eggs"
tests remain disabled as those tests fail on Linux when re-enabled and require
further investigation.

diff --git a/Engine/Plugins/Importers/USDImporter/Source/USDSchemas/USDSchemas.Build.cs b/Engine/Plugins/Importers/USDImporter/Source/USDSchemas/USDSchemas.Build.cs
--- a/Engine/Plugins/Importers/USDImporter/Source/USDSchemas/USDSchemas.Build.cs
+++ b/Engine/Plugins/Importers/USDImporter/Source/USDSchemas/USDSchemas.Build.cs
@@ -69,7 +69,10 @@ namespace UnrealBuildTool.Rules
 					}
 				);
 
-				if (Target.Platform == UnrealTargetPlatform.Win64)
+				// MaterialX via USD is only enabled on Windows and Linux for now while
+				// we work through issues affecting Mac.
+				if (Target.IsInPlatformGroup(UnrealPlatformGroup.Windows) ||
+					Target.IsInPlatformGroup(UnrealPlatformGroup.Unix))
 				{
 					PrivateDependencyModuleNames.AddRange(
 						new string[]
diff --git a/Engine/Plugins/Runtime/USDCore/Source/USDUtilities/USDUtilities.Build.cs b/Engine/Plugins/Runtime/USDCore/Source/USDUtilities/USDUtilities.Build.cs
--- a/Engine/Plugins/Runtime/USDCore/Source/USDUtilities/USDUtilities.Build.cs
+++ b/Engine/Plugins/Runtime/USDCore/Source/USDUtilities/USDUtilities.Build.cs
@@ -67,9 +67,10 @@ namespace UnrealBuildTool.Rules
 					}
 				);
 
-				// MaterialX via USD is only enabled on Windows for now while
-				// we work through issues affecting Mac and Linux.
-				if (Target.Platform == UnrealTargetPlatform.Win64)
+				// MaterialX via USD is only enabled on Windows and Linux for now while
+				// we work through issues affecting Mac.
+				if (Target.IsInPlatformGroup(UnrealPlatformGroup.Windows) ||
+					Target.IsInPlatformGroup(UnrealPlatformGroup.Unix))
 				{
 					// Public because MaterialX headers are #include'd in
 					// public USDUtilities headers.
