USDCore, USDImporter: re-enable MaterialX via USD on Mac

diff --git a/Engine/Plugins/Importers/USDImporter/Source/USDSchemas/USDSchemas.Build.cs b/Engine/Plugins/Importers/USDImporter/Source/USDSchemas/USDSchemas.Build.cs
--- a/Engine/Plugins/Importers/USDImporter/Source/USDSchemas/USDSchemas.Build.cs
+++ b/Engine/Plugins/Importers/USDImporter/Source/USDSchemas/USDSchemas.Build.cs
@@ -69,10 +69,9 @@ namespace UnrealBuildTool.Rules
 					}
 				);
 
-				// MaterialX via USD is only enabled on Windows and Linux for now while
-				// we work through issues affecting Mac.
 				if (Target.IsInPlatformGroup(UnrealPlatformGroup.Windows) ||
-					Target.IsInPlatformGroup(UnrealPlatformGroup.Unix))
+					Target.IsInPlatformGroup(UnrealPlatformGroup.Unix) ||
+					Target.Platform == UnrealTargetPlatform.Mac)
 				{
 					PrivateDependencyModuleNames.AddRange(
 						new string[]
diff --git a/Engine/Plugins/Runtime/USDCore/Source/USDUtilities/USDUtilities.Build.cs b/Engine/Plugins/Runtime/USDCore/Source/USDUtilities/USDUtilities.Build.cs
--- a/Engine/Plugins/Runtime/USDCore/Source/USDUtilities/USDUtilities.Build.cs
+++ b/Engine/Plugins/Runtime/USDCore/Source/USDUtilities/USDUtilities.Build.cs
@@ -67,10 +67,9 @@ namespace UnrealBuildTool.Rules
 					}
 				);
 
-				// MaterialX via USD is only enabled on Windows and Linux for now while
-				// we work through issues affecting Mac.
 				if (Target.IsInPlatformGroup(UnrealPlatformGroup.Windows) ||
-					Target.IsInPlatformGroup(UnrealPlatformGroup.Unix))
+					Target.IsInPlatformGroup(UnrealPlatformGroup.Unix) ||
+					Target.Platform == UnrealTargetPlatform.Mac)
 				{
 					// Public because MaterialX headers are #include'd in
 					// public USDUtilities headers.
