USD MaterialX:
- Add custom nodedefs from PXR_MTLX_PLUGIN_SEARCH_PATHS environment variable

#13625

diff --git a/Engine/Plugins/Runtime/USDCore/Source/USDUtilities/Private/USDMaterialXShaderGraph.cpp b/Engine/Plugins/Runtime/USDCore/Source/USDUtilities/Private/USDMaterialXShaderGraph.cpp
--- a/Engine/Plugins/Runtime/USDCore/Source/USDUtilities/Private/USDMaterialXShaderGraph.cpp
+++ b/Engine/Plugins/Runtime/USDCore/Source/USDUtilities/Private/USDMaterialXShaderGraph.cpp
@@ -96,6 +96,17 @@ FUsdMaterialXShaderGraph::FUsdMaterialXShaderGraph(const pxr::UsdPrim& UsdShadeM
 		mx::FileSearchPath MaterialXFolder{ TCHAR_TO_UTF8(*FPaths::Combine(FPaths::EngineDir(), TEXT("Binaries"), TEXT("ThirdParty"), TEXT("MaterialX"))) };
 		mx::StringSet LoadedLibs = mx::loadLibraries({ "libraries" }, MaterialXFolder, MaterialXLibrary);
 		Document->importLibrary(MaterialXLibrary);
+
+		// The path to the folder containing the custom node nodedefs should also be indicated with the PXR_MTLX_PLUGIN_SEARCH_PATHS environment variable 
+		std::string PxrMtlxPluginSearchPath = TCHAR_TO_UTF8(*FPlatformMisc::GetEnvironmentVariable(TEXT("PXR_MTLX_PLUGIN_SEARCH_PATHS")));
+		if(!PxrMtlxPluginSearchPath.empty())
+		{
+			mx::DocumentPtr PluginLibrary = mx::createDocument();
+			mx::FileSearchPath PluginFolder{ PxrMtlxPluginSearchPath };
+			mx::StringSet PluginLibs = mx::loadLibraries({}, PluginFolder, PluginLibrary);
+			Document->importLibrary(PluginLibrary);
+		}
+
 		CreateMaterial(UsdShadeMaterial, RenderContext);
 	}
 	catch (const std::exception& Exception)
