For merged module only enable python if it's set to be compiled by the target and in editor.

diff --git a/Engine/Plugins/Runtime/USDCore/Source/UnrealUSDWrapper/UnrealUSDWrapper.Build.cs b/Engine/Plugins/Runtime/USDCore/Source/UnrealUSDWrapper/UnrealUSDWrapper.Build.cs
--- a/Engine/Plugins/Runtime/USDCore/Source/UnrealUSDWrapper/UnrealUSDWrapper.Build.cs
+++ b/Engine/Plugins/Runtime/USDCore/Source/UnrealUSDWrapper/UnrealUSDWrapper.Build.cs
@@ -194,6 +194,11 @@ namespace UnrealBuildTool.Rules
 				bEnableUsdSdk = false;
 			}
 
+			if (!Target.bCompilePython)
+			{
+				bEnableUsdSdk = false;
+			}
+
 			// If you want to use USD in a monolithic target, you'll have to use the ANSI allocator.
 			// USD always uses the ANSI C allocators directly. In a DLL UE build (so not monolithic) we can just override the operators new and delete
 			// on each module with versions that use either the ANSI (so USD-compatible) allocators or the UE allocators (ModuleBoilerplate.h) when appropriate.
