For merged module only enable python if it's set to be compiled by the target and in editor.

diff --git a/Engine/Plugins/Runtime/USDCore/Source/UnrealUSDWrapper/UnrealUSDWrapper.Build.cs b/Engine/Plugins/Runtime/USDCore/Source/UnrealUSDWrapper/UnrealUSDWrapper.Build.cs
--- a/Engine/Plugins/Runtime/USDCore/Source/UnrealUSDWrapper/UnrealUSDWrapper.Build.cs
+++ b/Engine/Plugins/Runtime/USDCore/Source/UnrealUSDWrapper/UnrealUSDWrapper.Build.cs
@@ -228,6 +228,11 @@ namespace UnrealBuildTool.Rules
 				PublicDefinitions.Add("USD_FORCE_DISABLED=1");
 				bEnableUsdSdk = false;
 			}
+			else if (!Target.bCompilePython && Target.bMergeModules)
+			{
+				PublicDefinitions.Add("USD_FORCE_DISABLED=0");
+				bEnableUsdSdk = false;
+			}
 			else
 			{
 				PublicDefinitions.Add("USD_FORCE_DISABLED=0");
