For merged module only enable python if it's set to be compiled by the target and in editor.

#rnx

diff --git a/Engine/Plugins/Experimental/ChaosCachingUSD/Source/ChaosCachingUSD/ChaosCachingUSD.build.cs b/Engine/Plugins/Experimental/ChaosCachingUSD/Source/ChaosCachingUSD/ChaosCachingUSD.build.cs
--- a/Engine/Plugins/Experimental/ChaosCachingUSD/Source/ChaosCachingUSD/ChaosCachingUSD.build.cs
+++ b/Engine/Plugins/Experimental/ChaosCachingUSD/Source/ChaosCachingUSD/ChaosCachingUSD.build.cs
@@ -56,6 +56,11 @@ namespace UnrealBuildTool.Rules
 				bEnableUsdSdk = false;
 			}
 
+			if (!Target.bCompilePython)
+			{
+				bEnableUsdSdk = false;
+			}
+
 			bool bIsGameOrMonolithic = Target.LinkType == TargetLinkType.Monolithic || Target.Type == TargetType.Client || Target.Type == TargetType.Server || Target.Type == TargetType.Game;
 			if (bEnableUsdSdk &&
 				bIsGameOrMonolithic &&
diff --git a/Engine/Plugins/Experimental/ChaosFlesh/Source/ChaosFleshEditor/ChaosFleshEditor.Build.cs b/Engine/Plugins/Experimental/ChaosFlesh/Source/ChaosFleshEditor/ChaosFleshEditor.Build.cs
--- a/Engine/Plugins/Experimental/ChaosFlesh/Source/ChaosFleshEditor/ChaosFleshEditor.Build.cs
+++ b/Engine/Plugins/Experimental/ChaosFlesh/Source/ChaosFleshEditor/ChaosFleshEditor.Build.cs
@@ -112,6 +112,11 @@ namespace UnrealBuildTool.Rules
 				bEnableUsdSdk = false;
 			}
 
+			if (!Target.bCompilePython)
+			{
+				bEnableUsdSdk = false;
+			}
+
 			// If you want to use USD in a monolithic target, you'll have to use the ANSI allocator.
 			// USD always uses the ANSI C allocators directly. In a DLL UE build (so not monolithic) we can just override the operators new and delete
 			// on each module with versions that use either the ANSI (so USD-compatible) allocators or the UE allocators (ModuleBoilerplate.h) when appropriate.
diff --git a/Engine/Plugins/Experimental/ChaosFlesh/Source/ChaosFleshEngine/ChaosFleshEngine.Build.cs b/Engine/Plugins/Experimental/ChaosFlesh/Source/ChaosFleshEngine/ChaosFleshEngine.Build.cs
--- a/Engine/Plugins/Experimental/ChaosFlesh/Source/ChaosFleshEngine/ChaosFleshEngine.Build.cs
+++ b/Engine/Plugins/Experimental/ChaosFlesh/Source/ChaosFleshEngine/ChaosFleshEngine.Build.cs
@@ -119,6 +119,11 @@ namespace UnrealBuildTool.Rules
 				bEnableUsdSdk = false;
 			}
 
+			if (!Target.bCompilePython)
+			{
+				bEnableUsdSdk = false;
+			}
+
 			// If you want to use USD in a monolithic target, you'll have to use the ANSI allocator.
 			// USD always uses the ANSI C allocators directly. In a DLL UE build (so not monolithic) we can just override the operators new and delete
 			// on each module with versions that use either the ANSI (so USD-compatible) allocators or the UE allocators (ModuleBoilerplate.h) when appropriate.
diff --git a/Engine/Plugins/Importers/USDImporter/Source/USDExporter/USDExporter.Build.cs b/Engine/Plugins/Importers/USDImporter/Source/USDExporter/USDExporter.Build.cs
--- a/Engine/Plugins/Importers/USDImporter/Source/USDExporter/USDExporter.Build.cs
+++ b/Engine/Plugins/Importers/USDImporter/Source/USDExporter/USDExporter.Build.cs
@@ -86,6 +86,11 @@ namespace UnrealBuildTool.Rules
 				bEnableUsdSdk = false;
 			}
 
+			if (!Target.bCompilePython)
+			{
+				bEnableUsdSdk = false;
+			}
+
 			bool bIsGameOrMonolithic = Target.LinkType == TargetLinkType.Monolithic || Target.Type == TargetType.Client || Target.Type == TargetType.Server || Target.Type == TargetType.Game;
 			if (bEnableUsdSdk &&
 				bIsGameOrMonolithic &&
diff --git a/Engine/Plugins/Importers/USDImporter/Source/USDSchemas/USDSchemas.Build.cs b/Engine/Plugins/Importers/USDImporter/Source/USDSchemas/USDSchemas.Build.cs
--- a/Engine/Plugins/Importers/USDImporter/Source/USDSchemas/USDSchemas.Build.cs
+++ b/Engine/Plugins/Importers/USDImporter/Source/USDSchemas/USDSchemas.Build.cs
@@ -107,6 +107,11 @@ namespace UnrealBuildTool.Rules
 				bEnableUsdSdk = false;
 			}
 
+			if (!Target.bCompilePython)
+			{
+				bEnableUsdSdk = false;
+			}
+
 			bool bIsGameOrMonolithic = Target.LinkType == TargetLinkType.Monolithic || Target.Type == TargetType.Client || Target.Type == TargetType.Server || Target.Type == TargetType.Game;
 			if (bEnableUsdSdk &&
 				bIsGameOrMonolithic &&
diff --git a/Engine/Plugins/Importers/USDImporter/Source/USDStage/USDStage.Build.cs b/Engine/Plugins/Importers/USDImporter/Source/USDStage/USDStage.Build.cs
--- a/Engine/Plugins/Importers/USDImporter/Source/USDStage/USDStage.Build.cs
+++ b/Engine/Plugins/Importers/USDImporter/Source/USDStage/USDStage.Build.cs
@@ -97,6 +97,11 @@ namespace UnrealBuildTool.Rules
 				bEnableUsdSdk = false;
 			}
 
+			if (!Target.bCompilePython)
+			{
+				bEnableUsdSdk = false;
+			}
+
 			bool bIsGameOrMonolithic = Target.LinkType == TargetLinkType.Monolithic || Target.Type == TargetType.Client || Target.Type == TargetType.Server || Target.Type == TargetType.Game;
 			if (bEnableUsdSdk &&
 				bIsGameOrMonolithic &&
diff --git a/Engine/Plugins/Importers/USDImporter/Source/USDStageEditor/USDStageEditor.Build.cs b/Engine/Plugins/Importers/USDImporter/Source/USDStageEditor/USDStageEditor.Build.cs
--- a/Engine/Plugins/Importers/USDImporter/Source/USDStageEditor/USDStageEditor.Build.cs
+++ b/Engine/Plugins/Importers/USDImporter/Source/USDStageEditor/USDStageEditor.Build.cs
@@ -65,6 +65,11 @@ namespace UnrealBuildTool.Rules
 				bEnableUsdSdk = false;
 			}
 
+			if (!Target.bCompilePython)
+			{
+				bEnableUsdSdk = false;
+			}
+
 			bool bIsGameOrMonolithic = Target.LinkType == TargetLinkType.Monolithic || Target.Type == TargetType.Client || Target.Type == TargetType.Server || Target.Type == TargetType.Game;
 			if (bEnableUsdSdk &&
 				bIsGameOrMonolithic &&
diff --git a/Engine/Plugins/Importers/USDImporter/Source/USDStageEditorViewModels/USDStageEditorViewModels.Build.cs b/Engine/Plugins/Importers/USDImporter/Source/USDStageEditorViewModels/USDStageEditorViewModels.Build.cs
--- a/Engine/Plugins/Importers/USDImporter/Source/USDStageEditorViewModels/USDStageEditorViewModels.Build.cs
+++ b/Engine/Plugins/Importers/USDImporter/Source/USDStageEditorViewModels/USDStageEditorViewModels.Build.cs
@@ -59,6 +59,11 @@ namespace UnrealBuildTool.Rules
 				bEnableUsdSdk = false;
 			}
 
+			if (!Target.bCompilePython)
+			{
+				bEnableUsdSdk = false;
+			}
+
 			bool bIsGameOrMonolithic = Target.LinkType == TargetLinkType.Monolithic || Target.Type == TargetType.Client || Target.Type == TargetType.Server || Target.Type == TargetType.Game;
 			if (bEnableUsdSdk &&
 				bIsGameOrMonolithic &&
diff --git a/Engine/Plugins/Importers/USDImporter/Source/USDStageImporter/USDStageImporter.Build.cs b/Engine/Plugins/Importers/USDImporter/Source/USDStageImporter/USDStageImporter.Build.cs
--- a/Engine/Plugins/Importers/USDImporter/Source/USDStageImporter/USDStageImporter.Build.cs
+++ b/Engine/Plugins/Importers/USDImporter/Source/USDStageImporter/USDStageImporter.Build.cs
@@ -71,6 +71,11 @@ namespace UnrealBuildTool.Rules
 				bEnableUsdSdk = false;
 			}
 
+			if (!Target.bCompilePython)
+			{
+				bEnableUsdSdk = false;
+			}
+
 			bool bIsGameOrMonolithic = Target.LinkType == TargetLinkType.Monolithic || Target.Type == TargetType.Client || Target.Type == TargetType.Server || Target.Type == TargetType.Game;
 			if (bEnableUsdSdk &&
 				bIsGameOrMonolithic &&
diff --git a/Engine/Plugins/Importers/USDImporter/Source/USDTests/USDTests.Build.cs b/Engine/Plugins/Importers/USDImporter/Source/USDTests/USDTests.Build.cs
--- a/Engine/Plugins/Importers/USDImporter/Source/USDTests/USDTests.Build.cs
+++ b/Engine/Plugins/Importers/USDImporter/Source/USDTests/USDTests.Build.cs
@@ -60,6 +60,11 @@ namespace UnrealBuildTool.Rules
 				bEnableUsdSdk = false;
 			}
 
+			if (!Target.bCompilePython)
+			{
+				bEnableUsdSdk = false;
+			}
+
 			bool bIsGameOrMonolithic = Target.LinkType == TargetLinkType.Monolithic || Target.Type == TargetType.Client || Target.Type == TargetType.Server || Target.Type == TargetType.Game;
 			if (bEnableUsdSdk &&
 				bIsGameOrMonolithic &&
diff --git a/Engine/Plugins/Importers/USDImporterMDL/Source/USDImporterMDL/USDImporterMDL.Build.cs b/Engine/Plugins/Importers/USDImporterMDL/Source/USDImporterMDL/USDImporterMDL.Build.cs
--- a/Engine/Plugins/Importers/USDImporterMDL/Source/USDImporterMDL/USDImporterMDL.Build.cs
+++ b/Engine/Plugins/Importers/USDImporterMDL/Source/USDImporterMDL/USDImporterMDL.Build.cs
@@ -65,6 +65,11 @@ public class USDImporterMDL : ModuleRules
 			bEnableUsdSdk = false;
 		}
 
+		if (!Target.bCompilePython)
+		{
+			bEnableUsdSdk = false;
+		}
+
 		bool bIsGameOrMonolithic = Target.LinkType == TargetLinkType.Monolithic || Target.Type == TargetType.Client || Target.Type == TargetType.Server || Target.Type == TargetType.Game;
 		if (bEnableUsdSdk &&
 			bIsGameOrMonolithic &&
diff --git a/Engine/Plugins/Runtime/USDCore/Source/USDUtilities/USDUtilities.Build.cs b/Engine/Plugins/Runtime/USDCore/Source/USDUtilities/USDUtilities.Build.cs
--- a/Engine/Plugins/Runtime/USDCore/Source/USDUtilities/USDUtilities.Build.cs
+++ b/Engine/Plugins/Runtime/USDCore/Source/USDUtilities/USDUtilities.Build.cs
@@ -111,6 +111,11 @@ namespace UnrealBuildTool.Rules
 				bEnableUsdSdk = false;
 			}
 
+			if (!Target.bCompilePython)
+			{
+				bEnableUsdSdk = false;
+			}
+
 			bool bIsGameOrMonolithic = Target.LinkType == TargetLinkType.Monolithic || Target.Type == TargetType.Client || Target.Type == TargetType.Server || Target.Type == TargetType.Game;
 			if (bEnableUsdSdk &&
 				bIsGameOrMonolithic &&
diff --git a/Engine/Plugins/Runtime/USDCore/Source/UnrealUSDWrapper/UnrealUSDWrapper.Build.cs b/Engine/Plugins/Runtime/USDCore/Source/UnrealUSDWrapper/UnrealUSDWrapper.Build.cs
--- a/Engine/Plugins/Runtime/USDCore/Source/UnrealUSDWrapper/UnrealUSDWrapper.Build.cs
+++ b/Engine/Plugins/Runtime/USDCore/Source/UnrealUSDWrapper/UnrealUSDWrapper.Build.cs
@@ -194,6 +194,11 @@ namespace UnrealBuildTool.Rules
 				bEnableUsdSdk = false;
 			}
 
+			if (!Target.bCompilePython)
+			{
+				bEnableUsdSdk = false;
+			}
+
 			// If you want to use USD in a monolithic target, you'll have to use the ANSI allocator.
 			// USD always uses the ANSI C allocators directly. In a DLL UE build (so not monolithic) we can just override the operators new and delete
 			// on each module with versions that use either the ANSI (so USD-compatible) allocators or the UE allocators (ModuleBoilerplate.h) when appropriate.
@@ -228,11 +233,6 @@ namespace UnrealBuildTool.Rules
 				PublicDefinitions.Add("USD_FORCE_DISABLED=1");
 				bEnableUsdSdk = false;
 			}
-			else if (!Target.bCompilePython && Target.bMergeModules)
-			{
-				PublicDefinitions.Add("USD_FORCE_DISABLED=0");
-				bEnableUsdSdk = false;
-			}
 			else
 			{
 				PublicDefinitions.Add("USD_FORCE_DISABLED=0");
