USD: Fix wrong skeletal mesh influences when the skeleton contains multiple root bones.

diff --git a/Engine/Plugins/Interchange/Extensions/OpenUSD/Source/Import/Private/InterchangeUsdTranslator.cpp b/Engine/Plugins/Interchange/Extensions/OpenUSD/Source/Import/Private/InterchangeUsdTranslator.cpp
--- a/Engine/Plugins/Interchange/Extensions/OpenUSD/Source/Import/Private/InterchangeUsdTranslator.cpp
+++ b/Engine/Plugins/Interchange/Extensions/OpenUSD/Source/Import/Private/InterchangeUsdTranslator.cpp
@@ -6659,6 +6659,10 @@ TOptional<UE::Interchange::FMeshPayloadData> UInterchangeUSDTranslator::GetMeshP
 				OptionsCopy.AdditionalTransform = RootJointGlobalTransform;
 			}
 
+			// The skeletal mesh factory will already remap our bones by name from RetrieveAllSkeletalMeshPayloads() anyway,
+			// so we shouldn't offset anything ourselves
+			OptionsCopy.bOffsetJointIndicesForMultipleRootBones = false;
+
 			bSuccess = UE::InterchangeUsdTranslator::Private::GetSkeletalMeshPayloadData(
 				PayloadKey.UniqueId,
 				*ImplPtr,
diff --git a/Engine/Plugins/Runtime/USDCore/Source/USDUtilities/Private/USDGeomMeshConversion.cpp b/Engine/Plugins/Runtime/USDCore/Source/USDUtilities/Private/USDGeomMeshConversion.cpp
--- a/Engine/Plugins/Runtime/USDCore/Source/USDUtilities/Private/USDGeomMeshConversion.cpp
+++ b/Engine/Plugins/Runtime/USDCore/Source/USDUtilities/Private/USDGeomMeshConversion.cpp
@@ -1744,7 +1744,7 @@ namespace UE::UsdGeomMeshConversion::Private
 		// not), but naming it this way allows us to use it like an offset, which should make it easier to
 		// understand whenever it is used
 		// TODO: Can this be done simpler by just performing it after the joint remapping?
-		uint32 NumAdditionalBones = USDSkeletalDataConversion::HasMultipleRootBones(*SkeletonQuery) ? 1 : 0;
+		uint32 NumAdditionalBones = (Options.bOffsetJointIndicesForMultipleRootBones && USDSkeletalDataConversion::HasMultipleRootBones(*SkeletonQuery)) ? 1 : 0;
 		if (NumAdditionalBones > 0)
 		{
 			for (int32 Index = 0; Index < InOutMeshData.JointIndices.size(); ++Index)
@@ -3376,6 +3376,8 @@ namespace UsdToUnreal
 		, TimeCode(pxr::UsdTimeCode::EarliestTime())
 		, bMergeIdenticalMaterialSlots(true)
 		, SubdivisionLevel(0)
+		, ImportPrimvars(EImportPrimvar::Standard)
+		, bOffsetJointIndicesForMultipleRootBones(true)
 	{
 	}
 }
diff --git a/Engine/Plugins/Runtime/USDCore/Source/USDUtilities/Public/USDGeomMeshConversion.h b/Engine/Plugins/Runtime/USDCore/Source/USDUtilities/Public/USDGeomMeshConversion.h
--- a/Engine/Plugins/Runtime/USDCore/Source/USDUtilities/Public/USDGeomMeshConversion.h
+++ b/Engine/Plugins/Runtime/USDCore/Source/USDUtilities/Public/USDGeomMeshConversion.h
@@ -112,7 +112,10 @@ namespace UsdToUnreal
 		int32 SubdivisionLevel;
 
 		// How to store the primvars in the MeshDescription
-		EImportPrimvar ImportPrimvars = EImportPrimvar::Standard;
+		EImportPrimvar ImportPrimvars;
+
+		// Whether to add an additional offset to vertex join indices if the skinned mesh's skeleton has multiple root bones
+		bool bOffsetJointIndicesForMultipleRootBones;
 	};
 
 	/**
