USD: Fix support for UDIMs.

Fix materials not being marked as used if the user prim doesn't have the MaterialBindingAPI *and* it's using a specific material purpose.

diff --git a/Engine/Plugins/Runtime/USDCore/Source/USDUtilities/Private/Objects/USDInfoCache.cpp b/Engine/Plugins/Runtime/USDCore/Source/USDUtilities/Private/Objects/USDInfoCache.cpp
--- a/Engine/Plugins/Runtime/USDCore/Source/USDUtilities/Private/Objects/USDInfoCache.cpp
+++ b/Engine/Plugins/Runtime/USDCore/Source/USDUtilities/Private/Objects/USDInfoCache.cpp
@@ -1276,108 +1276,98 @@ namespace UE::USDInfoCache::Private
 			TMap<UE::FSdfPath, TSet<UE::FSdfPath>> NewMaterialUsers;
 
 			pxr::UsdShadeMaterialBindingAPI BindingAPI{UsdPrim};
+
+			// Check for material users via collections-based material bindings
 			if (BindingAPI || bPossibleInheritedBindings)
 			{
-				// Check for material users via collections-based material bindings
+				// When retrieving the relationships directly we'll always need to check the universal render context
+				// manually, as it won't automatically "compute the fallback" for us like when we ComputeBoundMaterial()
+				std::unordered_set<pxr::TfToken, pxr::TfHash> MaterialPurposeTokens{
+					MaterialPurposeToken,
+					pxr::UsdShadeTokens->universalRenderContext
+				};
+				for (const pxr::TfToken& SomeMaterialPurposeToken : MaterialPurposeTokens)
 				{
-					// When retrieving the relationships directly we'll always need to check the universal render context
-					// manually, as it won't automatically "compute the fallback" for us like when we ComputeBoundMaterial()
-					std::unordered_set<pxr::TfToken, pxr::TfHash> MaterialPurposeTokens{
-						MaterialPurposeToken,
-						pxr::UsdShadeTokens->universalRenderContext
-					};
-					for (const pxr::TfToken& SomeMaterialPurposeToken : MaterialPurposeTokens)
+					// Each one of those relationships must have two targets: A collection, and a material
+					for (const pxr::UsdRelationship& Rel : BindingAPI.GetCollectionBindingRels(SomeMaterialPurposeToken))
 					{
-						// Each one of those relationships must have two targets: A collection, and a material
-						for (const pxr::UsdRelationship& Rel : BindingAPI.GetCollectionBindingRels(SomeMaterialPurposeToken))
-						{
-							const pxr::SdfPath* CollectionPath = nullptr;
-							const pxr::SdfPath* MaterialPath = nullptr;
+						const pxr::SdfPath* CollectionPath = nullptr;
+						const pxr::SdfPath* MaterialPath = nullptr;
 
-							std::vector<pxr::SdfPath> PathVector;
-							if (Rel.GetTargets(&PathVector))
+						std::vector<pxr::SdfPath> PathVector;
+						if (Rel.GetTargets(&PathVector))
+						{
+							for (const pxr::SdfPath& Path : PathVector)
 							{
-								for (const pxr::SdfPath& Path : PathVector)
+								if (Path.IsPrimPath())
 								{
-									if (Path.IsPrimPath())
-									{
-										MaterialPath = &Path;
-									}
-									else if (Path.IsPropertyPath())
-									{
-										CollectionPath = &Path;
-									}
+									MaterialPath = &Path;
+								}
+								else if (Path.IsPropertyPath())
+								{
+									CollectionPath = &Path;
 								}
 							}
+						}
 
-							if (!CollectionPath || !MaterialPath || PathVector.size() != 2)
-							{
-								// Emit this warning here as USD doesn't seem to and just seems to just ignores this relationship instead
-								USD_LOG_USERWARNING(FText::Format(
-									LOCTEXT(
-										"InvalidCollection",
-										"Prim '{0}' describes a collection-based material binding, but the relationship '{1}' is invalid: It should "
-										"contain exactly one Material path and one path to a collection relationship"
-									),
-									FText::FromString(Info.PrimPath.GetString()),
-									FText::FromString(UsdToUnreal::ConvertToken(Rel.GetName()))
-								));
-								continue;
-							}
+						if (!CollectionPath || !MaterialPath || PathVector.size() != 2)
+						{
+							// Emit this warning here as USD doesn't seem to and just seems to just ignores this relationship instead
+							USD_LOG_USERWARNING(FText::Format(
+								LOCTEXT(
+									"InvalidCollection",
+									"Prim '{0}' describes a collection-based material binding, but the relationship '{1}' is invalid: It should "
+									"contain exactly one Material path and one path to a collection relationship"
+								),
+								FText::FromString(Info.PrimPath.GetString()),
+								FText::FromString(UsdToUnreal::ConvertToken(Rel.GetName()))
+							));
+							continue;
+						}
 
-							if (pxr::UsdCollectionAPI Collection = pxr::UsdCollectionAPI::Get(Stage, *CollectionPath))
-							{
-								TSet<UE::FSdfPath>& MaterialUsers = NewMaterialUsers.FindOrAdd(UE::FSdfPath{*MaterialPath});
+						if (pxr::UsdCollectionAPI Collection = pxr::UsdCollectionAPI::Get(Stage, *CollectionPath))
+						{
+							TSet<UE::FSdfPath>& MaterialUsers = NewMaterialUsers.FindOrAdd(UE::FSdfPath{*MaterialPath});
 
-								std::set<pxr::SdfPath> IncludedPaths = Collection.ComputeIncludedPaths(Collection.ComputeMembershipQuery(), Stage);
-								for (const pxr::SdfPath& IncludedPath : IncludedPaths)
-								{
-									TryAddMaterialUser(UE::FSdfPath{IncludedPath}, MaterialUsers);
-								}
-							}
-							else
+							std::set<pxr::SdfPath> IncludedPaths = Collection.ComputeIncludedPaths(Collection.ComputeMembershipQuery(), Stage);
+							for (const pxr::SdfPath& IncludedPath : IncludedPaths)
 							{
-								USD_LOG_USERWARNING(FText::Format(
-									LOCTEXT(
-										"MissingCollection",
-										"Failed to find collection at path '{0}' when processing collection-based material bindings on prim '{1}'"
-									),
-									FText::FromString(UsdToUnreal::ConvertPath(CollectionPath->GetPrimPath())),
-									FText::FromString(Info.PrimPath.GetString())
-								));
+								TryAddMaterialUser(UE::FSdfPath{IncludedPath}, MaterialUsers);
 							}
 						}
+						else
+						{
+							USD_LOG_USERWARNING(FText::Format(
+								LOCTEXT(
+									"MissingCollection",
+									"Failed to find collection at path '{0}' when processing collection-based material bindings on prim '{1}'"
+								),
+								FText::FromString(UsdToUnreal::ConvertPath(CollectionPath->GetPrimPath())),
+								FText::FromString(Info.PrimPath.GetString())
+							));
+						}
 					}
 				}
-
-				// Check for material bindings directly for this prim
-				if (pxr::UsdShadeMaterial ShadeMaterial = BindingAPI.ComputeBoundMaterial(MaterialPurposeToken))
-				{
-					bPrimHasInheritableMaterialBindings = true;
-
-					TSet<UE::FSdfPath>& MaterialUsers = NewMaterialUsers.FindOrAdd(UE::FSdfPath{ShadeMaterial.GetPrim().GetPath()});
-					TryAddMaterialUser(Info.PrimPath, MaterialUsers);
-				}
 			}
-			// Temporary fallback for prims that don't have the MaterialBindingAPI but do have the relationship.
-			// USD will emit a warning for these though
-			else if (pxr::UsdRelationship Relationship = UsdPrim.GetRelationship(pxr::UsdShadeTokens->materialBinding))
-			{
-				pxr::SdfPathVector Targets;
-				Relationship.GetTargets(&Targets);
 
-				if (Targets.size() > 0)
-				{
-					const pxr::SdfPath& TargetMaterialPrimPath = Targets[0];
-					pxr::UsdPrim MaterialPrim = Stage->GetPrimAtPath(TargetMaterialPrimPath);
-					if (pxr::UsdShadeMaterial ShadeMaterial = pxr::UsdShadeMaterial{MaterialPrim})
-					{
-						bPrimHasInheritableMaterialBindings = true;
+			// Check for material bindings directly for this prim.
+			//
+			// Note how we don't actually check the BindingAPI here: This is intentional as it will cover the case where the prim
+			// doesn't have the MaterialBindingAPI applied but still has the material relationships, which is also how
+			// UsdUtils::GetPrimMaterialAssignments also operates. The MaterialBindingAPI schema code will still work even if the
+			// schema is not actually applied to the prim.
+			//
+			// We used to check for a relationship directly here using the "material:binding" token, but that is inferior to using ComputeBoundMaterial
+			// as it will not use the MaterialPurposeToken: If the prim only had "material:binding:preview", we would fail to find a relationship when
+			// just querying with "material:binding"
+			//
+			// USD will emit a warning for the missing MaterialBindingAPI though.
+			if (pxr::UsdShadeMaterial ShadeMaterial = BindingAPI.ComputeBoundMaterial(MaterialPurposeToken))
+			{
+				bPrimHasInheritableMaterialBindings = true;
 
-						TSet<UE::FSdfPath>& MaterialUsers = NewMaterialUsers.FindOrAdd(UE::FSdfPath{TargetMaterialPrimPath});
-						TryAddMaterialUser(Info.PrimPath, MaterialUsers);
-					}
-				}
+				TSet<UE::FSdfPath>& MaterialUsers = NewMaterialUsers.FindOrAdd(UE::FSdfPath{ShadeMaterial.GetPrim().GetPath()});
+				TryAddMaterialUser(Info.PrimPath, MaterialUsers);
 			}
 
 			for (const TPair<UE::FSdfPath, TSet<UE::FSdfPath>>& NewMaterialToUsers : NewMaterialUsers)
diff --git a/Engine/Plugins/Runtime/USDCore/Source/USDUtilities/Private/USDShadeConversion.cpp b/Engine/Plugins/Runtime/USDCore/Source/USDUtilities/Private/USDShadeConversion.cpp
--- a/Engine/Plugins/Runtime/USDCore/Source/USDUtilities/Private/USDShadeConversion.cpp
+++ b/Engine/Plugins/Runtime/USDCore/Source/USDUtilities/Private/USDShadeConversion.cpp
@@ -26,6 +26,7 @@
 #include "IImageWrapper.h"
 #include "IImageWrapperModule.h"
 #include "ImageUtils.h"
+#include "Internationalization/Regex.h"
 #include "Materials/Material.h"
 #include "Materials/MaterialExpressionConstant.h"
 #include "Materials/MaterialExpressionConstant4Vector.h"
@@ -39,6 +40,7 @@
 #include "Modules/ModuleManager.h"
 #include "RenderUtils.h"
 #include "TextureResource.h"
+#include "UDIMUtilities.h"
 
 #if WITH_EDITOR
 #include "Factories/TextureFactory.h"
@@ -1248,7 +1250,8 @@ namespace UE::UsdShadeConversion::Private
 		TextureGroup Group,
 		EObjectFlags ObjectFlags,
 		UObject* Outer,
-		bool bForceLinear
+		bool bForceLinear,
+		bool bIsUDIM
 	)
 	{
 		UTexture* Texture = nullptr;
@@ -1268,17 +1271,19 @@ namespace UE::UsdShadeConversion::Private
 		// To maintain existing behavior, ColorSpaceMode is left at auto if linear is not requested
 		TextureFactory->ColorSpaceMode = bForceLinear ? ETextureSourceColorSpace::Linear : ETextureSourceColorSpace::Auto;
 
-		const bool bIsSupportedUdimTexture = ResolvedTexturePath.Contains(TEXT("<UDIM>"));
-		if (bIsSupportedUdimTexture)
+		if (bIsUDIM)
 		{
-			FString BaseFileName = FPaths::GetBaseFilename(ResolvedTexturePath);
-
-			FString BaseFileNameBeforeUdim;
-			FString BaseFileNameAfterUdim;
-			BaseFileName.Split(TEXT("<UDIM>"), &BaseFileNameBeforeUdim, &BaseFileNameAfterUdim);
+			const static FRegexPattern RegexPattern(TEXT(R"((.+[._])(\d{4}|<UDIM>)(.*))"));
+			const FString BaseFileName = FPaths::GetBaseFilename(ResolvedTexturePath);
+			FRegexMatcher RegexMatcher(RegexPattern, BaseFileName);
+			if (RegexMatcher.FindNext())
+			{
+				FString BaseFileNameBeforeUdim = RegexMatcher.GetCaptureGroup(1);
+				FString BaseFileNameAfterUdim = RegexMatcher.GetCaptureGroup(3);
 
-			FString UdimRegexPattern = FString::Printf(TEXT(R"((%s)(\d{4})(%s))"), *BaseFileNameBeforeUdim, *BaseFileNameAfterUdim);
-			TextureFactory->UdimRegexPattern = MoveTemp(UdimRegexPattern);
+				FString UdimRegexPattern = FString::Printf(TEXT(R"((%s)(\d{4})(%s))"), *BaseFileNameBeforeUdim, *BaseFileNameAfterUdim);
+				TextureFactory->UdimRegexPattern = MoveTemp(UdimRegexPattern);
+			}
 		}
 
 		if (!ResolvedTexturePath.IsEmpty())
@@ -2073,12 +2078,14 @@ namespace UE::UsdShadeConversion::Private
 
 			const EObjectFlags DesiredFlags = RF_Public | RF_Standalone | RF_Transactional;
 
+			const bool bIsUDIM = TextureParameter->bIsUDIM || TextureParameter->TextureFilePath.Contains(TEXT("<UDIM>"));
+
 			bool bCreatedTexture = false;
 			UTexture* Texture = TexturesCache->GetOrCreateCustomCachedAsset<UTexture2D>(
 				PrefixedTextureHash,
 				DesiredTextureName,
 				DesiredFlags,
-				[TextureParameter, TexturesCache, bSRGB](UPackage* Outer, FName SanitizedName, EObjectFlags DesiredFlags)
+				[TextureParameter, TexturesCache, bSRGB, bIsUDIM](UPackage* Outer, FName SanitizedName, EObjectFlags DesiredFlags)
 				{
 					return UsdUtils::CreateTexture(
 						TextureParameter->TextureFilePath,
@@ -2086,7 +2093,8 @@ namespace UE::UsdShadeConversion::Private
 						TextureParameter->Group,
 						DesiredFlags,
 						Outer,
-						!bSRGB
+						!bSRGB,
+						bIsUDIM
 					);
 				},
 				&bCreatedTexture
@@ -2096,8 +2104,11 @@ namespace UE::UsdShadeConversion::Private
 			{
 				TextureParameter->Texture = Texture;
 
+				const bool bNeedVirtualTextureStreaming = bIsUDIM;
+
 				// The texture resource needs to be updated only if the following settings have changed from their default values
 				bool bNeedUpdateResource = false;
+				bNeedUpdateResource |= bNeedVirtualTextureStreaming != Texture->VirtualTextureStreaming;
 				bNeedUpdateResource |= TextureParameter->AddressX != TextureAddress::TA_Wrap;
 				bNeedUpdateResource |= TextureParameter->AddressY != TextureAddress::TA_Wrap;
 #if !WITH_EDITOR
@@ -2111,6 +2122,11 @@ namespace UE::UsdShadeConversion::Private
 #if WITH_EDITOR
 					Texture->PreEditChange(nullptr);
 #endif	  // WITH_EDITOR
+					if (bNeedVirtualTextureStreaming)
+					{
+						Texture->VirtualTextureStreaming = true;
+					}
+
 					Texture->SRGB = bSRGB;
 					Texture->CompressionSettings = CompressionSettings;
 					if (UTexture2D* Texture2D = Cast<UTexture2D>(Texture))
@@ -2778,7 +2794,8 @@ UTexture* UsdUtils::CreateTexture(
 	TextureGroup Group,
 	EObjectFlags Flags,
 	UObject* Outer,
-	bool bForceLinear
+	bool bForceLinear,
+	bool bIsUDIMPath
 )
 {
 	TRACE_CPUPROFILER_EVENT_SCOPE(UsdUtils::CreateTexture);
@@ -2786,7 +2803,7 @@ UTexture* UsdUtils::CreateTexture(
 	// Standalone game does have WITH_EDITOR defined, but it can't use the texture factories, so we need to check this instead
 	if (GIsEditor)
 	{
-		return UsdShadeConversionImpl::CreateTextureWithEditor(ResolvedTexturePath, SanitizedName, Group, Flags, Outer, bForceLinear);
+		return UsdShadeConversionImpl::CreateTextureWithEditor(ResolvedTexturePath, SanitizedName, Group, Flags, Outer, bForceLinear, bIsUDIMPath);
 	}
 	else
 	{
diff --git a/Engine/Plugins/Runtime/USDCore/Source/USDUtilities/Public/USDShadeConversion.h b/Engine/Plugins/Runtime/USDCore/Source/USDUtilities/Public/USDShadeConversion.h
--- a/Engine/Plugins/Runtime/USDCore/Source/USDUtilities/Public/USDShadeConversion.h
+++ b/Engine/Plugins/Runtime/USDCore/Source/USDUtilities/Public/USDShadeConversion.h
@@ -331,7 +331,8 @@ namespace UsdUtils
 		TextureGroup Group = TEXTUREGROUP_World,
 		EObjectFlags Flags = RF_Public | RF_Standalone | RF_Transactional,
 		UObject* Outer = GetTransientPackage(),
-		bool bForceLinear = false
+		bool bForceLinear = false,
+		bool bIsUDIMPath = false
 	);
 
 	/** Checks if this texture needs virtual textures and emits a warning if it is disabled for the project */
diff --git a/Engine/Plugins/Runtime/USDCore/Source/USDUtilities/USDUtilities.Build.cs b/Engine/Plugins/Runtime/USDCore/Source/USDUtilities/USDUtilities.Build.cs
--- a/Engine/Plugins/Runtime/USDCore/Source/USDUtilities/USDUtilities.Build.cs
+++ b/Engine/Plugins/Runtime/USDCore/Source/USDUtilities/USDUtilities.Build.cs
@@ -41,10 +41,10 @@ namespace UnrealBuildTool.Rules
 					"OpenSubdiv",
 					"RenderCore",
 					"RHI", // So that we can use GMaxRHIFeatureLevel when force-loading textures before baking materials
+					"SkeletalMeshDescription",
 					"Slate",
 					"SlateCore",
 					"StaticMeshDescription",
-					"SkeletalMeshDescription"
 				}
 			);
 
@@ -63,6 +63,7 @@ namespace UnrealBuildTool.Rules
 						"MeshUtilities",
 						"MessageLog",
 						"PropertyEditor",
+						"TextureUtilitiesCommon",
 						"UnrealEd",
 					}
 				);
