USD Interchange: Use stage root layer filename instead of "/" for pseudoroot actor display label.

diff --git a/Engine/Plugins/Interchange/Extensions/OpenUSD/Source/Import/Private/InterchangeUsdTranslator.cpp b/Engine/Plugins/Interchange/Extensions/OpenUSD/Source/Import/Private/InterchangeUsdTranslator.cpp
--- a/Engine/Plugins/Interchange/Extensions/OpenUSD/Source/Import/Private/InterchangeUsdTranslator.cpp
+++ b/Engine/Plugins/Interchange/Extensions/OpenUSD/Source/Import/Private/InterchangeUsdTranslator.cpp
@@ -5112,9 +5112,17 @@ namespace UE::InterchangeUsdTranslator::Private
 	{
 		TRACE_CPUPROFILER_EVENT_SCOPE(Traverse)
 
-		const FString SceneNodeUid = MakeNodeUid(Prim.GetPrimPath().GetString());
-		const FString DisplayLabel(Prim.GetName().ToString());
 		const FName TypeName = Prim.GetTypeName();
+		const FString SceneNodeUid = MakeNodeUid(Prim.GetPrimPath().GetString());
+		FString DisplayLabel;
+		if (Prim.IsPseudoRoot())
+		{
+			DisplayLabel = Prim.GetStage().GetRootLayer().GetDisplayName();
+		}
+		else
+		{
+			DisplayLabel = Prim.HasAuthoredDisplayName() ? Prim.GetDisplayName() : Prim.GetName().ToString();
+		}
 
 		// Do this before generating other nodes as they may need the updated info
 		Info.UpdateWithCurrentPrim(Prim);
diff --git a/Engine/Plugins/Interchange/Extensions/OpenUSD/Source/Import/Public/InterchangeUSDPipeline.h b/Engine/Plugins/Interchange/Extensions/OpenUSD/Source/Import/Public/InterchangeUSDPipeline.h
--- a/Engine/Plugins/Interchange/Extensions/OpenUSD/Source/Import/Public/InterchangeUSDPipeline.h
+++ b/Engine/Plugins/Interchange/Extensions/OpenUSD/Source/Import/Public/InterchangeUSDPipeline.h
@@ -42,8 +42,8 @@ public:
 	int32 SubdivisionLevel;
 
 	/**
-	 * The translator always emits a scene node for the stage pseudoroot. Setting this option to false means we will remove it from
-	 * the factory nodes, making top-level prims lead to top-level actors and components.
+	 * The translator always emits a scene node for the stage pseudoroot named after the stage's root layer filename.
+	 * Setting this option to false means we will remove it from the factory nodes, making top-level prims lead to top-level actors and components.
 	 */
 	UPROPERTY(EditAnywhere, BlueprintReadWrite, Category = "USD")
 	bool bImportPseudoRoot;
diff --git a/Engine/Plugins/Runtime/USDCore/Source/UnrealUSDWrapper/Private/UsdWrappers/UsdPrim.cpp b/Engine/Plugins/Runtime/USDCore/Source/UnrealUSDWrapper/Private/UsdWrappers/UsdPrim.cpp
--- a/Engine/Plugins/Runtime/USDCore/Source/UnrealUSDWrapper/Private/UsdWrappers/UsdPrim.cpp
+++ b/Engine/Plugins/Runtime/USDCore/Source/UnrealUSDWrapper/Private/UsdWrappers/UsdPrim.cpp
@@ -232,6 +232,44 @@ namespace UE
 		return Result;
 	}
 
+	FString FUsdPrim::GetDisplayName() const
+	{
+		FString Result;
+
+#if USE_USD_SDK
+		Result = UTF8_TO_TCHAR(Impl->PxrUsdPrim.Get().GetDisplayName().c_str());
+#endif	  // #if USE_USD_SDK
+
+		return Result;
+	}
+
+	bool FUsdPrim::SetDisplayName(const FString& Name) const
+	{
+#if USE_USD_SDK
+		return Impl->PxrUsdPrim.Get().SetDisplayName(TCHAR_TO_UTF8(*Name));
+#else
+		return false;
+#endif	  // #if USE_USD_SDK
+	}
+
+	bool FUsdPrim::ClearDisplayName() const
+	{
+#if USE_USD_SDK
+		return Impl->PxrUsdPrim.Get().ClearDisplayName();
+#else
+		return false;
+#endif	  // #if USE_USD_SDK
+	}
+
+	bool FUsdPrim::HasAuthoredDisplayName() const
+	{
+#if USE_USD_SDK
+		return Impl->PxrUsdPrim.Get().HasAuthoredDisplayName();
+#else
+		return false;
+#endif	  // #if USE_USD_SDK
+	}
+
 	bool FUsdPrim::SetSpecifier(ESdfSpecifier Specifier)
 	{
 #if USE_USD_SDK
diff --git a/Engine/Plugins/Runtime/USDCore/Source/UnrealUSDWrapper/Public/UsdWrappers/UsdPrim.h b/Engine/Plugins/Runtime/USDCore/Source/UnrealUSDWrapper/Public/UsdWrappers/UsdPrim.h
--- a/Engine/Plugins/Runtime/USDCore/Source/UnrealUSDWrapper/Public/UsdWrappers/UsdPrim.h
+++ b/Engine/Plugins/Runtime/USDCore/Source/UnrealUSDWrapper/Public/UsdWrappers/UsdPrim.h
@@ -88,6 +88,11 @@ namespace UE
 		TMap<FString, UE::FVtValue> GetAllMetadata() const;
 		TMap<FString, UE::FVtValue> GetAllAuthoredMetadata() const;
 
+		FString GetDisplayName() const;
+		bool SetDisplayName(const FString& Name) const;
+		bool ClearDisplayName() const;
+		bool HasAuthoredDisplayName() const;
+
 		// Wrapped pxr::UsdPrim functions, refer to the USD SDK documentation
 	public:
 		bool SetSpecifier(ESdfSpecifier Specifier);
