[USD] Triangulate quads and n-gons when generating mesh descriptions

Changelist 44617758 modified FStaticMeshOperations::AppendMeshDescriptions() so that mesh descriptions are no longer automatically triangulated. This
exposed an issue where recomputed normals are only correct when working with triangulated meshes.

This fix restores triangulation for USD-generated mesh descriptions, ensuring that recomputed normals remain accurate. A new console variable,
USD.TriangulateQuadsAndNgons, can be set to false to disable this behavior and preserve quads and n-gons in the resulting mesh descriptions.

diff --git a/Engine/Plugins/Runtime/USDCore/Source/USDUtilities/Private/USDGeomMeshConversion.cpp b/Engine/Plugins/Runtime/USDCore/Source/USDUtilities/Private/USDGeomMeshConversion.cpp
--- a/Engine/Plugins/Runtime/USDCore/Source/USDUtilities/Private/USDGeomMeshConversion.cpp
+++ b/Engine/Plugins/Runtime/USDCore/Source/USDUtilities/Private/USDGeomMeshConversion.cpp
@@ -134,6 +134,13 @@ static FAutoConsoleVariableRef CVarSkipMeshTangentComputation(
 	TEXT("Skip computing tangents for meshes. With meshes with a huge numer of vertices, it can take a very long time to compute them.")
 );
 
+static bool GTriangulateQuadsAndNgons = true;
+static FAutoConsoleVariableRef CVarTriangulateQuadsAndNgons(
+	TEXT("USD.TriangulateQuadsAndNgons"),
+	GTriangulateQuadsAndNgons,
+	TEXT("Triangulate quad and ngon faces when generating static and skeletal meshes.")
+);
+
 namespace UE::UsdGeomMeshConversion::Private
 {
 	static const FString DisplayColorID = TEXT("!DisplayColor");
@@ -3358,7 +3365,35 @@ namespace UE::UsdGeomMeshConversion::Private
 				*InMeshData.SourcePrimPath
 			);
 		}
+				
+		// Triangulate quads and ngons in order to avoid bad recomputed normals (UE-330573) 
+		if (GTriangulateQuadsAndNgons)
+		{
+			TArray<FPolygonID> NonTrianglesToDelete;
+			for (const FPolygonID PolygonID : OutMeshDescription.Polygons().GetElementIDs())
+			{
+				TArrayView<const FTriangleID> TriangleIDs = OutMeshDescription.GetPolygonTriangles(PolygonID);
+				if (TriangleIDs.Num() < 2)
+				{
+					continue;
+				}
 
+				const FPolygonGroupID PolygonGroupID = OutMeshDescription.GetPolygonPolygonGroup(PolygonID);
+				for (const FTriangleID TriangleID : TriangleIDs)
+				{
+					TArrayView<const FVertexInstanceID> VertexInstanceIDs = OutMeshDescription.GetTriangleVertexInstances(TriangleID);
+					OutMeshDescription.CreateTriangle(PolygonGroupID, VertexInstanceIDs);
+				}
+
+				NonTrianglesToDelete.Add(PolygonID);
+			}
+
+			for (FPolygonID PolygonID : NonTrianglesToDelete)
+			{
+				OutMeshDescription.DeletePolygon(PolygonID);
+			}
+		}
+		
 		return true;
 	}
 }	 // namespace UE::UsdGeomMeshConversion::Private
