[USD] Fix for crash in recently added mesh triangluation code

diff --git a/Engine/Plugins/Runtime/USDCore/Source/USDUtilities/Private/USDGeomMeshConversion.cpp b/Engine/Plugins/Runtime/USDCore/Source/USDUtilities/Private/USDGeomMeshConversion.cpp
--- a/Engine/Plugins/Runtime/USDCore/Source/USDUtilities/Private/USDGeomMeshConversion.cpp
+++ b/Engine/Plugins/Runtime/USDCore/Source/USDUtilities/Private/USDGeomMeshConversion.cpp
@@ -3365,7 +3365,7 @@ namespace UE::UsdGeomMeshConversion::Private
 				*InMeshData.SourcePrimPath
 			);
 		}
-				
+
 		// Triangulate quads and ngons in order to avoid bad recomputed normals (UE-330573) 
 		if (GTriangulateQuadsAndNgons)
 		{
@@ -3381,7 +3381,12 @@ namespace UE::UsdGeomMeshConversion::Private
 				const FPolygonGroupID PolygonGroupID = OutMeshDescription.GetPolygonPolygonGroup(PolygonID);
 				for (const FTriangleID TriangleID : TriangleIDs)
 				{
-					TArrayView<const FVertexInstanceID> VertexInstanceIDs = OutMeshDescription.GetTriangleVertexInstances(TriangleID);
+					FVertexInstanceID VertexInstanceIDs[] { 
+						OutMeshDescription.GetTriangleVertexInstance(TriangleID, 0),
+						OutMeshDescription.GetTriangleVertexInstance(TriangleID, 1), 
+						OutMeshDescription.GetTriangleVertexInstance(TriangleID, 2) 
+					};
+				
 					OutMeshDescription.CreateTriangle(PolygonGroupID, VertexInstanceIDs);
 				}
 
