USD: Register the search paths set in the PXR_AR_DEFAULT_SEARCH_PATH environment variable to make them available to the default asset resolver.

diff --git a/Engine/Plugins/Runtime/USDCore/Source/UnrealUSDWrapper/Private/UnrealUSDWrapper.cpp b/Engine/Plugins/Runtime/USDCore/Source/UnrealUSDWrapper/Private/UnrealUSDWrapper.cpp
--- a/Engine/Plugins/Runtime/USDCore/Source/UnrealUSDWrapper/Private/UnrealUSDWrapper.cpp
+++ b/Engine/Plugins/Runtime/USDCore/Source/UnrealUSDWrapper/Private/UnrealUSDWrapper.cpp
@@ -1286,7 +1286,16 @@ public:
 		}
 
 		// Set the default search path for USD's default resolver using any path specified in the settings.
-		const TArray<FDirectoryPath> SearchPath = ProjectSettings->DefaultResolverSearchPath;
+		TArray<FDirectoryPath> SearchPath = ProjectSettings->DefaultResolverSearchPath;
+
+		FString ArSearchPathEnvVar = FPlatformMisc::GetEnvironmentVariable(TEXT("PXR_AR_DEFAULT_SEARCH_PATH"));
+		if (!ArSearchPathEnvVar.IsEmpty())
+		{
+			TArray<FString> SearchPaths;
+			ArSearchPathEnvVar.ParseIntoArray(SearchPaths, FPlatformMisc::GetPathVarDelimiter());
+			SearchPath.Append(SearchPaths);
+		}
+
 		UnrealUSDWrapper::SetDefaultResolverDefaultSearchPath(SearchPath);
 
 #endif	  // USE_USD_SDK
