USD: Fixed crashes due to missing scoped allocs needed with the custom allocator disabled.

#rnx

diff --git a/Engine/Plugins/Runtime/USDCore/Source/USDUtilities/Private/USDGeomMeshConversion.cpp b/Engine/Plugins/Runtime/USDCore/Source/USDUtilities/Private/USDGeomMeshConversion.cpp
--- a/Engine/Plugins/Runtime/USDCore/Source/USDUtilities/Private/USDGeomMeshConversion.cpp
+++ b/Engine/Plugins/Runtime/USDCore/Source/USDUtilities/Private/USDGeomMeshConversion.cpp
@@ -6940,6 +6940,8 @@ bool UsdUtils::GetPointInstancerTransforms(
 
 void UsdUtils::GetPointInstancerInstancesData(const pxr::UsdPrim& Prim, FUsdPointInstancerInstancesData& InstancesData)
 {
+	FScopedUsdAllocs Allocs;
+
 	pxr::UsdGeomPointInstancer PointInstancer(Prim);
 	if (!PointInstancer)
 	{
diff --git a/Engine/Plugins/Runtime/USDCore/Source/UnrealUSDWrapper/Private/UsdWrappers/UsdPrim.cpp b/Engine/Plugins/Runtime/USDCore/Source/UnrealUSDWrapper/Private/UsdWrappers/UsdPrim.cpp
--- a/Engine/Plugins/Runtime/USDCore/Source/UnrealUSDWrapper/Private/UsdWrappers/UsdPrim.cpp
+++ b/Engine/Plugins/Runtime/USDCore/Source/UnrealUSDWrapper/Private/UsdWrappers/UsdPrim.cpp
@@ -204,6 +204,8 @@ namespace UE
 		// avoid issues due to casing differences
 		TMap<FString, UE::FVtValue> Result;
 #if USE_USD_SDK
+		FScopedUsdAllocs Allocs;
+
 		std::map<pxr::TfToken, pxr::VtValue, pxr::TfDictionaryLessThan> UsdMetadata = Impl->PxrUsdPrim.Get().GetAllMetadata();
 		Result.Reserve(UsdMetadata.size());
 
@@ -220,6 +222,8 @@ namespace UE
 	{
 		TMap<FString, UE::FVtValue> Result;
 #if USE_USD_SDK
+		FScopedUsdAllocs Allocs;
+
 		std::map<pxr::TfToken, pxr::VtValue, pxr::TfDictionaryLessThan> UsdMetadata = Impl->PxrUsdPrim.Get().GetAllAuthoredMetadata();
 		Result.Reserve(UsdMetadata.size());
 
diff --git a/Engine/Plugins/Runtime/USDCore/Source/UnrealUSDWrapper/Private/UsdWrappers/VtValue.cpp b/Engine/Plugins/Runtime/USDCore/Source/UnrealUSDWrapper/Private/UsdWrappers/VtValue.cpp
--- a/Engine/Plugins/Runtime/USDCore/Source/UnrealUSDWrapper/Private/UsdWrappers/VtValue.cpp
+++ b/Engine/Plugins/Runtime/USDCore/Source/UnrealUSDWrapper/Private/UsdWrappers/VtValue.cpp
@@ -819,6 +819,8 @@ namespace UE
 	void FVtValue::Set(const T& Value)
 	{
 #if USE_USD_SDK
+		FScopedUsdAllocs UsdAllocs;
+
 		pxr::VtValue& VtValue = GetUsdValue();
 
 		if constexpr (std::is_same_v<T, int64>)
