USD: Fixed potential crash when exporting a geometry cache to USD due to missing velocities.
There was a bug when importing a geometry cache from USD that lost the imported velocities.

#rnx

diff --git a/Engine/Plugins/Runtime/USDCore/Source/USDUtilities/Private/USDGeomMeshConversion.cpp b/Engine/Plugins/Runtime/USDCore/Source/USDUtilities/Private/USDGeomMeshConversion.cpp
--- a/Engine/Plugins/Runtime/USDCore/Source/USDUtilities/Private/USDGeomMeshConversion.cpp
+++ b/Engine/Plugins/Runtime/USDCore/Source/USDUtilities/Private/USDGeomMeshConversion.cpp
@@ -4927,6 +4927,8 @@ namespace UE::UsdGeometryCacheConversion::Private
 		InOutFlattenedMeshData.TangentsX.Append(InMeshData.TangentsX);
 		InOutFlattenedMeshData.TangentsZ.Append(InMeshData.TangentsZ);
 		InOutFlattenedMeshData.Colors.Append(InMeshData.Colors);
+		InOutFlattenedMeshData.MotionVectors.Append(InMeshData.MotionVectors);
+		InOutFlattenedMeshData.ImportedVertexNumbers.Append(InMeshData.ImportedVertexNumbers);
 
 		// ... and adjusting the indices with the proper offset
 		InOutFlattenedMeshData.Indices.Reserve(InOutFlattenedMeshData.Indices.Num() + InMeshData.Indices.Num());
@@ -4948,10 +4950,12 @@ namespace UE::UsdGeometryCacheConversion::Private
 		}
 
 		// Also merge the VertexInfo attributes that are checked when converting the MeshData
+		InOutFlattenedMeshData.VertexInfo.bHasTangentX |= InMeshData.VertexInfo.bHasTangentX;
 		InOutFlattenedMeshData.VertexInfo.bHasTangentZ |= InMeshData.VertexInfo.bHasTangentZ;
 		InOutFlattenedMeshData.VertexInfo.bHasUV0 |= InMeshData.VertexInfo.bHasUV0;
 		InOutFlattenedMeshData.VertexInfo.bHasColor0 |= InMeshData.VertexInfo.bHasColor0;
 		InOutFlattenedMeshData.VertexInfo.bHasMotionVectors |= InMeshData.VertexInfo.bHasMotionVectors;
+		InOutFlattenedMeshData.VertexInfo.bHasImportedVertexNumbers |= InMeshData.VertexInfo.bHasImportedVertexNumbers;
 	}
 
 	FGeometryCacheMeshData GetFlattenedGeometryCacheMeshData(const UGeometryCache* GeometryCache, int32 FrameIndex)
@@ -5081,11 +5085,11 @@ namespace UE::UsdGeometryCacheConversion::Private
 		}
 		bool HasVelocities() const override
 		{
-			return MeshData.VertexInfo.bHasMotionVectors;
+			return MeshData.VertexInfo.bHasMotionVectors && !MeshData.MotionVectors.IsEmpty();
 		}
 		FVector3f GetVelocity(int32 Index) const override
 		{
-			return MeshData.MotionVectors[Index];
+			return MeshData.MotionVectors.IsValidIndex(Index) ? MeshData.MotionVectors[Index] : FVector3f::ZeroVector;
 		}
 		int32 GetNumSections() const override
 		{
@@ -5471,6 +5475,8 @@ namespace UE::UsdGeometryCacheConversion::Private
 					if (ExportContext.SlotNames.IsValidIndex(SectionIndex))
 					{
 						SectionName = ExportContext.SlotNames[SectionIndex].ToString();
+						// Make sure the SectionName contains no invalid characters for a prim identifier, otherwise the prim creation will fail
+						SectionName = UsdUtils::SanitizeUsdIdentifier(*SectionName);
 						SectionName = UsdUnreal::ObjectUtils::GetUniqueName(SectionName, UsedSectionNames);
 						UsedSectionNames.Add(SectionName);
 					}
