USD: Prevent loss of transform opinions in some edge cases when writing out attribute values to USD.

diff --git a/Engine/Plugins/Runtime/USDCore/Source/USDUtilities/Private/USDPrimConversion.cpp b/Engine/Plugins/Runtime/USDCore/Source/USDUtilities/Private/USDPrimConversion.cpp
--- a/Engine/Plugins/Runtime/USDCore/Source/USDUtilities/Private/USDPrimConversion.cpp
+++ b/Engine/Plugins/Runtime/USDCore/Source/USDUtilities/Private/USDPrimConversion.cpp
@@ -163,9 +163,6 @@ namespace UE::USDPrimConversion::Private
 
 		// References: private constructor for UsdGeomXformOp that can receive a UsdPrim and UsdGeomXformable::AddXformOp
 
-		// Clear the existing xform op order for this prim on this layer
-		Xformable.ClearXformOpOrder();
-
 		// Find details about the transform attribute related to the default transform type xform op
 		pxr::TfToken TransformAttrName = pxr::UsdGeomXformOp::GetOpName(pxr::UsdGeomXformOp::TypeTransform);
 		const pxr::SdfValueTypeName& TransformAttrTypeName = pxr::UsdGeomXformOp::GetValueTypeName(
@@ -194,6 +191,69 @@ namespace UE::USDPrimConversion::Private
 			return {};
 		}
 
+		// If we have other xformOps, let's convert them to matrix xforms as well.
+		// This is important as we can end up in edge cases like UE-221328, where we may have some timeSamples for different
+		// xformOps that are fully lost when we author a single default opinion for a matrix op, as we'll also force the
+		// xformOpOrder to just use our matrix op
+		bool bResetsXformStack = false;
+		std::vector<pxr::UsdGeomXformOp> ExistingOps = Xformable.GetOrderedXformOps(&bResetsXformStack);
+		if (ExistingOps.size() > 0)
+		{
+			// If there's exactly one op and it's already xformOp:transform, then we don't need to do any conversion
+			bool bNeedConversion = false;
+			if (ExistingOps.size() > 1)
+			{
+				bNeedConversion = true;
+			}
+			else if (ExistingOps[0].GetOpType() != pxr::UsdGeomXformOp::TypeTransform || ExistingOps[0].IsInverseOp())
+			{
+				bNeedConversion = true;
+			}
+
+			if (bNeedConversion)
+			{
+				// Convert default opinion if any of the Ops have one
+				bool bHasDefault = false;
+				for (const pxr::UsdGeomXformOp& Op : ExistingOps)
+				{
+					const pxr::UsdAttribute& OpAttr = Op;
+					bHasDefault = OpAttr.GetResolveInfo(pxr::UsdTimeCode::Default()).HasAuthoredValue();
+					if (bHasDefault)
+					{
+						break;
+					}
+				}
+				if (bHasDefault)
+				{
+					// TODO: Proper handling of bResetsXformStack
+					pxr::GfMatrix4d Transform;
+					bool bResetsXformStackDefault = false;
+					if (Xformable.GetLocalTransformation(&Transform, &bResetsXformStackDefault, pxr::UsdTimeCode::Default()))
+					{
+						TransformAttr.Set(Transform, pxr::UsdTimeCode::Default());
+					}
+				}
+
+				// Convert timeSamples
+				std::vector<double> TimeSamples;
+				if (Xformable.GetTimeSamples(&TimeSamples) && TimeSamples.size() > 0)
+				{
+					for (double TimeSample : TimeSamples)
+					{
+						pxr::GfMatrix4d Transform;
+						bool bResetsXformStackTimeSample = false;
+						if (Xformable.GetLocalTransformation(&Transform, &bResetsXformStackTimeSample, TimeSample))
+						{
+							TransformAttr.Set(Transform, TimeSample);
+						}
+					}
+				}
+			}
+		}
+
+		// Clear the existing xform op order for this prim on this layer
+		Xformable.ClearXformOpOrder();
+
 		// Store the Op name on an array that will be our new op order value
 		pxr::VtTokenArray NewOps;
 		NewOps.push_back(NewOp.GetOpName());
@@ -4566,8 +4626,6 @@ bool UnrealToUsd::CreateComponentPropertyBaker(
 				return false;
 			}
 
-			Xformable.CreateXformOpOrderAttr();
-
 			// Clear existing transform data and leave just one Transform op there
 			pxr::UsdGeomXformOp TransformOp = UE::USDPrimConversion::Private::ForceMatrixXform(Xformable);
 			if (!TransformOp)
@@ -5491,8 +5549,6 @@ UnrealToUsd::FPropertyTrackWriter UnrealToUsd::CreatePropertyTrackWriter(
 			{
 				if (pxr::UsdGeomXformable Xformable{UsdPrim})
 				{
-					Xformable.CreateXformOpOrderAttr();
-
 					if (pxr::UsdGeomXformOp TransformOp = UE::USDPrimConversion::Private::ForceMatrixXform(Xformable))
 					{
 						Attr = TransformOp.GetAttr();
