USD: Prevent valid UTF8 letters from being sanitized away to underscores in some operations.

diff --git a/Engine/Plugins/Runtime/USDCore/Source/USDUtilities/Private/USDConversionUtils.cpp b/Engine/Plugins/Runtime/USDCore/Source/USDUtilities/Private/USDConversionUtils.cpp
--- a/Engine/Plugins/Runtime/USDCore/Source/USDUtilities/Private/USDConversionUtils.cpp
+++ b/Engine/Plugins/Runtime/USDCore/Source/USDUtilities/Private/USDConversionUtils.cpp
@@ -70,6 +70,7 @@
 #include "USDIncludesStart.h"
 #include "pxr/base/tf/stringUtils.h"
 #include "pxr/base/tf/token.h"
+#include "pxr/base/tf/unicodeUtils.h"
 #include "pxr/usd/kind/registry.h"
 #include "pxr/usd/sdf/copyUtils.h"
 #include "pxr/usd/sdf/schema.h"
@@ -3093,10 +3094,55 @@ FString UsdUtils::SanitizeUsdIdentifier(const TCHAR* InIdentifier)
 #if USE_USD_SDK
 	FScopedUsdAllocs Allocs;
 
-	std::string UsdInName = UnrealToUsd::ConvertString(InIdentifier).Get();
-	std::string UsdValidName = pxr::TfMakeValidIdentifier(UsdInName);
+	// We used to use pxr::TfMakeValidIdentifier() here, but it seems that that sanitizes away UTF8
+	// characters, which are otherwise supported.
+	//
+	// Now we implement something based on _IsValidIdentifier() from USD's sdf\path.cpp instead,
+	// as we need to actually produce a valid identifier and not only check if it's valid or not
+
+	std::string Identifier = UnrealToUsd::ConvertString(InIdentifier).Get();
+	if (Identifier.empty())
+	{
+		Identifier = "_";
+	}
+	else
+	{
+		std::string Sanitized;
+		Sanitized.reserve(Identifier.size());
+		std::stringstream Stream{Sanitized};
+
+		std::string_view View{Identifier};
+		pxr::TfUtf8CodePointIterator Iter{std::cbegin(View), std::cend(View)};
+
+		// First codepoint has a different rule
+		const pxr::TfUtf8CodePoint& FirstCodePoint = *Iter;
+		if (FirstCodePoint == pxr::TfUtf8CodePointFromAscii('_') || pxr::TfIsUtf8CodePointXidStart(FirstCodePoint))
+		{
+			Stream << FirstCodePoint;
+		}
+		else
+		{
+			Stream << '_';
+		}
+
+		// Sanitize remaining codepoints
+		for (++Iter; Iter != pxr::TfUtf8CodePointIterator::PastTheEndSentinel{}; ++Iter)
+		{
+			const pxr::TfUtf8CodePoint& CodePoint = *Iter;
+			if (TfIsUtf8CodePointXidContinue(CodePoint))
+			{
+				Stream << CodePoint;
+			}
+			else
+			{
+				Stream << '_';
+			}
+		}
+
+		Identifier = Stream.str();
+	}
 
-	return UsdToUnreal::ConvertString(UsdValidName);
+	return UsdToUnreal::ConvertString(Identifier);
 #else
 	return InIdentifier;
 #endif	  // USE_USD_SDK
