USD: Fix handling of nested purposes

diff --git a/Engine/Plugins/Importers/USDImporter/Source/USDSchemas/Private/USDGeomMeshTranslator.cpp b/Engine/Plugins/Importers/USDImporter/Source/USDSchemas/Private/USDGeomMeshTranslator.cpp
--- a/Engine/Plugins/Importers/USDImporter/Source/USDSchemas/Private/USDGeomMeshTranslator.cpp
+++ b/Engine/Plugins/Importers/USDImporter/Source/USDSchemas/Private/USDGeomMeshTranslator.cpp
@@ -468,6 +468,12 @@ namespace UsdGeomMeshTranslatorImpl
 			return;
 		}
 
+		// Ignore meshes from disabled purposes
+		if (!EnumHasAllFlags(Options.PurposesToLoad, IUsdPrim::GetPurpose(MeshPrim)))
+		{
+			return;
+		}
+
 		UE::FUsdStage Stage = MeshPrim.GetStage();
 		UE::FSdfPath Path = MeshPrim.GetPrimPath();
 
@@ -1636,12 +1642,6 @@ FBuildStaticMeshTaskChain::FBuildStaticMeshTaskChain(
 
 void FBuildStaticMeshTaskChain::SetupTasks()
 {
-	// Ignore meshes from disabled purposes
-	if (!EnumHasAllFlags(Context->PurposesToLoad, IUsdPrim::GetPurpose(GetPrim())))
-	{
-		return;
-	}
-
 	// Create static mesh (Main thread)
 	Do(ESchemaTranslationLaunchPolicy::Sync,
 	   [this]()
diff --git a/Engine/Plugins/Interchange/Extensions/OpenUSD/Source/Import/Private/InterchangeUSDPipeline.cpp b/Engine/Plugins/Interchange/Extensions/OpenUSD/Source/Import/Private/InterchangeUSDPipeline.cpp
--- a/Engine/Plugins/Interchange/Extensions/OpenUSD/Source/Import/Private/InterchangeUSDPipeline.cpp
+++ b/Engine/Plugins/Interchange/Extensions/OpenUSD/Source/Import/Private/InterchangeUSDPipeline.cpp
@@ -65,42 +65,25 @@ namespace UE::InterchangeUsdPipeline::Private
 	{
 		TSet<FString> DisabledTranslatedNodes;
 
-		TFunction<void(const FString&, bool)> RecursiveCollectFilteredTranslatedNodes = nullptr;
-		RecursiveCollectFilteredTranslatedNodes = [&NodeContainer,
-												   AllowedPurposes,
-												   &RecursiveCollectFilteredTranslatedNodes,
-												   &DisabledTranslatedNodes](const FString& NodeUid, bool bDisableSubtree)
-		{
-			const UInterchangeBaseNode* Node = NodeContainer.GetNode(NodeUid);
-			if (!Node)
-			{
-				return;
-			}
-
-			if (!bDisableSubtree)
+		NodeContainer.IterateNodesOfType<UInterchangeSceneNode>(
+			[&DisabledTranslatedNodes, AllowedPurposes](const FString& NodeUid, UInterchangeSceneNode* SceneNode)
 			{
-				EUsdPurpose AuthoredPurpose = EUsdPurpose::Default;
-				bool bHasPurpose = Node->GetInt32Attribute(UE::Interchange::USD::GeometryPurposeIdentifier, (int32&)AuthoredPurpose);
-				if (bHasPurpose && !EnumHasAllFlags(AllowedPurposes, AuthoredPurpose))
+				bool bDisabledPurpose = false;
 				{
-					// Purpose inheritance according to UsdGeomImageable::ComputePurposeInfo seems to be:
-					// - Authored purposes inherit down to prims without any purpose;
-					// - If a prim has any purpose authored on them, that is their "computed purpose", and that is inherited to its children.
-					//
-					// We don't care much about the actual purpose then: We just need to check if any prim/node has a purpose
-					// that is not allowed. If that is the case, we turn off the entire subtree
-					bDisableSubtree = true;
+					EUsdPurpose AuthoredPurpose = EUsdPurpose::Default;
+					bool bHasPurpose = SceneNode->GetInt32Attribute(UE::Interchange::USD::GeometryPurposeIdentifier, (int32&)AuthoredPurpose);
+					if (bHasPurpose && !EnumHasAllFlags(AllowedPurposes, AuthoredPurpose))
+					{
+						bDisabledPurpose = true;
+					}
 				}
-			}
 
-			if (bDisableSubtree)
-			{
-				// Mark the scene translated node as disabled
-				DisabledTranslatedNodes.Add(NodeUid);
-
-				// Mark the asset translated node as disabled
-				if (const UInterchangeSceneNode* SceneNode = Cast<UInterchangeSceneNode>(Node))
+				if (bDisabledPurpose)
 				{
+					// Mark the scene translated node as disabled
+					DisabledTranslatedNodes.Add(NodeUid);
+
+					// Mark the asset translated node as disabled
 					FString AssetNodeUid;
 					if (SceneNode->GetCustomAssetInstanceUid(AssetNodeUid))
 					{
@@ -108,14 +91,7 @@ namespace UE::InterchangeUsdPipeline::Private
 					}
 				}
 			}
-
-			for (const FString& ChildNodeUid : NodeContainer.GetNodeChildrenUids(NodeUid))
-			{
-				RecursiveCollectFilteredTranslatedNodes(ChildNodeUid, bDisableSubtree);
-			}
-		};
-		const bool bDisableSubtree = false;
-		RecursiveCollectFilteredTranslatedNodes(GetPseudoRootTranslatedNodeUid(), bDisableSubtree);
+		);
 
 		// Disable any factory node that targets/references one of our disabled scene nodes.
 		//
diff --git a/Engine/Plugins/Runtime/USDCore/Source/USDUtilities/Private/USDConversionUtils.cpp b/Engine/Plugins/Runtime/USDCore/Source/USDUtilities/Private/USDConversionUtils.cpp
--- a/Engine/Plugins/Runtime/USDCore/Source/USDUtilities/Private/USDConversionUtils.cpp
+++ b/Engine/Plugins/Runtime/USDCore/Source/USDUtilities/Private/USDConversionUtils.cpp
@@ -3260,7 +3260,7 @@ TArray<UE::FUsdPrim> UsdUtils::GetVisibleChildren(const UE::FUsdPrim& Prim, EUsd
 	TFunction<void(const pxr::UsdPrim& Prim)> RecursivelyCollectVisibleMeshes;
 	RecursivelyCollectVisibleMeshes = [&RecursivelyCollectVisibleMeshes, &VisiblePrims, AllowedPurposes](const pxr::UsdPrim& Prim)
 	{
-		if (!Prim || !EnumHasAllFlags(AllowedPurposes, IUsdPrim::GetPurpose(Prim)))
+		if (!Prim)
 		{
 			return;
 		}
@@ -3278,7 +3278,13 @@ TArray<UE::FUsdPrim> UsdUtils::GetVisibleChildren(const UE::FUsdPrim& Prim, EUsd
 			}
 		}
 
-		VisiblePrims.Add(UE::FUsdPrim{Prim});
+		// We want to only collect meshes from visible purposes, but unlike for the visibility attribute itself,
+		// we don't want to stop traversal due to a hidden purpose, as a child prim may have a different purpose
+		// to its parent
+		if (EnumHasAllFlags(AllowedPurposes, IUsdPrim::GetPurpose(Prim)))
+		{
+			VisiblePrims.Add(UE::FUsdPrim{Prim});
+		}
 
 		for (const pxr::UsdPrim& ChildPrim : Prim.GetFilteredChildren(pxr::UsdTraverseInstanceProxies()))
 		{
diff --git a/Engine/Plugins/Runtime/USDCore/Source/USDUtilities/Private/USDGeomMeshConversion.cpp b/Engine/Plugins/Runtime/USDCore/Source/USDUtilities/Private/USDGeomMeshConversion.cpp
--- a/Engine/Plugins/Runtime/USDCore/Source/USDUtilities/Private/USDGeomMeshConversion.cpp
+++ b/Engine/Plugins/Runtime/USDCore/Source/USDUtilities/Private/USDGeomMeshConversion.cpp
@@ -919,12 +919,6 @@ namespace UE::UsdGeomMeshConversion::Private
 		bool bIsInsideSkelRoot
 	)
 	{
-		// Ignore meshes from disabled purposes
-		if (!EnumHasAllFlags(Options.PurposesToLoad, IUsdPrim::GetPurpose(Prim)))
-		{
-			return true;
-		}
-
 		if (!bSkipVisibility)
 		{
 			// Ignore invisible child meshes.
@@ -980,25 +974,39 @@ namespace UE::UsdGeomMeshConversion::Private
 		// the original AdditionalTransform in place. The alternative is to copy the entire options object...
 		TGuardValue<FTransform> Guard{Options.AdditionalTransform, ChildTransform};
 
+		// Ignore meshes from disabled purposes, but keep traversing regardless.
+		// This because a child purposes can differ from a parent purposes (i.e. we could want to skip a parent mesh, but get a child mesh)
+		const bool bPurposeIsVisible = EnumHasAllFlags(Options.PurposesToLoad, IUsdPrim::GetPurpose(Prim));
+
 		if (pxr::UsdGeomMesh Mesh = pxr::UsdGeomMesh(Prim))
 		{
 			// We never want to glob up *skinned* meshes inside SkelRoots, as those presumably will be handled by the
 			// UsdSkelSkeletonTranslator and the skeletal data code path.
-			if (!GIsEditor || !bIsInsideSkelRoot || !Prim.HasAPI<pxr::UsdSkelBindingAPI>())
+			if (bPurposeIsVisible && (!GIsEditor || !bIsInsideSkelRoot || !Prim.HasAPI<pxr::UsdSkelBindingAPI>()))
 			{
 				bSuccess = UsdToUnreal::ConvertGeomMesh(Mesh, OutMeshDescription, OutMaterialAssignments, Options);
 			}
 		}
 		// Check for primitive types directly (don't just check for Gprim here because if we get something else like a UsdVolVolume or
 		// UsdGeomBasisCurves ConvertGeomPrimitive will return false and we will dump our MeshDescription)
-		else if (Prim.IsA<pxr::UsdGeomCapsule>() || Prim.IsA<pxr::UsdGeomCone>() || Prim.IsA<pxr::UsdGeomCube>() || Prim.IsA<pxr::UsdGeomCylinder>()
-				 || Prim.IsA<pxr::UsdGeomPlane>() || Prim.IsA<pxr::UsdGeomSphere>())
+		else if (Prim.IsA<pxr::UsdGeomCapsule>() ||
+				 Prim.IsA<pxr::UsdGeomCone>() ||
+				 Prim.IsA<pxr::UsdGeomCube>() ||
+				 Prim.IsA<pxr::UsdGeomCylinder>() ||
+				 Prim.IsA<pxr::UsdGeomPlane>() ||
+				 Prim.IsA<pxr::UsdGeomSphere>())
 		{
-			bSuccess = UsdToUnreal::ConvertGeomPrimitive(Prim, OutMeshDescription, OutMaterialAssignments, Options);
+			if (bPurposeIsVisible)
+			{
+				bSuccess = UsdToUnreal::ConvertGeomPrimitive(Prim, OutMeshDescription, OutMaterialAssignments, Options);
+			}
 		}
 		else if (pxr::UsdGeomPointInstancer PointInstancer = pxr::UsdGeomPointInstancer{Prim})
 		{
-			bSuccess = UsdToUnreal::ConvertPointInstancerToMesh(PointInstancer, OutMeshDescription, OutMaterialAssignments, Options);
+			if (bPurposeIsVisible)
+			{
+				bSuccess = UsdToUnreal::ConvertPointInstancerToMesh(PointInstancer, OutMeshDescription, OutMaterialAssignments, Options);
+			}
 
 			// We never want to step into point instancers when fetching prims for drawing
 			bTraverseChildren = false;
@@ -1042,11 +1050,6 @@ namespace UE::UsdGeomMeshConversion::Private
 
 		// This should always replicate the same traversal pattern of RecursivelyCollapseChildMeshes
 
-		if (!EnumHasAllFlags(Options.PurposesToLoad, IUsdPrim::GetPurpose(Prim)))
-		{
-			return;
-		}
-
 		if (!bIsFirstPrim)
 		{
 			if (pxr::UsdGeomImageable UsdGeomImageable = pxr::UsdGeomImageable(Prim))
@@ -1062,33 +1065,39 @@ namespace UE::UsdGeomMeshConversion::Private
 			}
 		}
 
+		const bool bPurposeIsVisible = EnumHasAllFlags(Options.PurposesToLoad, IUsdPrim::GetPurpose(Prim));
+
 		bool bTraverseChildren = true;
 
 		if (pxr::UsdGeomPointInstancer PointInstancer = pxr::UsdGeomPointInstancer{Prim})
 		{
-			pxr::SdfPathVector PrototypePaths;
-			if (!PointInstancer.GetPrototypesRel().GetTargets(&PrototypePaths))
+			if (bPurposeIsVisible)
 			{
-				return;
-			}
-
-			pxr::UsdStageRefPtr Stage = Prim.GetStage();
-			for (const pxr::SdfPath& PrototypePath : PrototypePaths)
-			{
-				pxr::UsdPrim PrototypeUsdPrim = Stage->GetPrimAtPath(PrototypePath);
-				if (!PrototypeUsdPrim)
+				pxr::SdfPathVector PrototypePaths;
+				if (!PointInstancer.GetPrototypesRel().GetTargets(&PrototypePaths))
 				{
-					continue;
+					return;
 				}
 
-				const bool bChildIsFirstPrim = false;
-				RecursivelyCollectPrimvars(PrototypeUsdPrim, Options, InOutAllPrimvars, InOutPreferredPrimvars, bChildIsFirstPrim);
+				pxr::UsdStageRefPtr Stage = Prim.GetStage();
+				for (const pxr::SdfPath& PrototypePath : PrototypePaths)
+				{
+					pxr::UsdPrim PrototypeUsdPrim = Stage->GetPrimAtPath(PrototypePath);
+					if (!PrototypeUsdPrim)
+					{
+						continue;
+					}
+
+					const bool bChildIsFirstPrim = false;
+					RecursivelyCollectPrimvars(PrototypeUsdPrim, Options, InOutAllPrimvars, InOutPreferredPrimvars, bChildIsFirstPrim);
+				}
 			}
 
 			// We never want to step into point instancers when fetching prims for drawing
 			bTraverseChildren = false;
 		}
-		if (Prim)
+
+		if (Prim && bPurposeIsVisible)
 		{
 			TArray<TUsdStore<pxr::UsdGeomPrimvar>> Primvars = UsdUtils::GetUVSetPrimvars(Prim, TNumericLimits<int32>::Max());
 
diff --git a/Engine/Plugins/Runtime/USDCore/Source/UnrealUSDWrapper/Private/UnrealUSDWrapper.cpp b/Engine/Plugins/Runtime/USDCore/Source/UnrealUSDWrapper/Private/UnrealUSDWrapper.cpp
--- a/Engine/Plugins/Runtime/USDCore/Source/UnrealUSDWrapper/Private/UnrealUSDWrapper.cpp
+++ b/Engine/Plugins/Runtime/USDCore/Source/UnrealUSDWrapper/Private/UnrealUSDWrapper.cpp
@@ -231,10 +231,6 @@ EUsdPurpose IUsdPrim::GetPurpose(const UsdPrim& Prim, bool bComputed)
 	UsdGeomImageable Geom(Prim);
 	if (Geom)
 	{
-		// Use compute purpose because it depends on the hierarchy:
-		// "If the purpose of </RootPrim> is set to "render", then the effective purpose
-		// of </RootPrim/ChildPrim> will be "render" even if that prim has a different
-		// authored value for purpose."
 		TfToken Purpose;
 		if (bComputed)
 		{
