USD: Add tooltips to the options on EUsdCollisionType

diff --git a/Engine/Plugins/Runtime/USDCore/Source/UnrealUSDWrapper/Public/UnrealUSDWrapper.h b/Engine/Plugins/Runtime/USDCore/Source/UnrealUSDWrapper/Public/UnrealUSDWrapper.h
--- a/Engine/Plugins/Runtime/USDCore/Source/UnrealUSDWrapper/Public/UnrealUSDWrapper.h
+++ b/Engine/Plugins/Runtime/USDCore/Source/UnrealUSDWrapper/Public/UnrealUSDWrapper.h
@@ -153,13 +153,28 @@ enum class EUsdRootMotionHandling : uint8
 UENUM()
 enum class EUsdCollisionType : uint8
 {
-	None,	 // no approximation so equivalent to CTF_UseComplexAsSimple
+	// No approximation so equivalent to CTF_UseComplexAsSimple
+	None,
+
+	// Breaks up concave meshes into multiple, simple convex collider primitives
 	ConvexDecomposition,
+
+	// Single convex hull primitive around the mesh
 	ConvexHull,
+
+	// Single sphere collider
 	Sphere,
+
+	// Single cube collider
 	Cube,
+
+	// Behaves the same as ConvexDecomposition in Unreal, it's just listed separately as it's a separate value in USD
 	MeshSimplification,
+
+	// Single capsule collider
 	Capsule,
+
+	// Use a separate, dedicated collider mesh
 	CustomMesh
 };
 
