USD: Prevent the creation of internal references targetting the default prim

diff --git a/Engine/Plugins/Importers/USDImporter/Source/USDStageEditor/Private/SUSDStageTreeView.cpp b/Engine/Plugins/Importers/USDImporter/Source/USDStageEditor/Private/SUSDStageTreeView.cpp
--- a/Engine/Plugins/Importers/USDImporter/Source/USDStageEditor/Private/SUSDStageTreeView.cpp
+++ b/Engine/Plugins/Importers/USDImporter/Source/USDStageEditor/Private/SUSDStageTreeView.cpp
@@ -1191,6 +1191,7 @@ void SUsdStageTreeView::OnAddReference()
 	if (OptionsPtr->bInternalReference)
 	{
 		ReferencedLayerPath = TEXT("");
+		OptionsPtr->bUseDefaultPrim = false;
 	}
 	else if (FPaths::IsRelative(ReferencedLayerPath))
 	{
@@ -1272,6 +1273,7 @@ void SUsdStageTreeView::OnAddPayload()
 	if (OptionsPtr->bInternalReference)
 	{
 		ReferencedLayerPath = TEXT("");
+		OptionsPtr->bUseDefaultPrim = false;
 	}
 	else if (FPaths::IsRelative(ReferencedLayerPath))
 	{
diff --git a/Engine/Plugins/Runtime/USDCore/Source/USDClasses/Public/USDReferenceOptions.h b/Engine/Plugins/Runtime/USDCore/Source/USDClasses/Public/USDReferenceOptions.h
--- a/Engine/Plugins/Runtime/USDCore/Source/USDClasses/Public/USDReferenceOptions.h
+++ b/Engine/Plugins/Runtime/USDCore/Source/USDClasses/Public/USDReferenceOptions.h
@@ -18,15 +18,15 @@ public:
 	bool bInternalReference = false;
 
 	// File to use as the reference
-	UPROPERTY(EditAnywhere, BlueprintReadWrite, Category = "Options", meta = (EditCondition = "!bInternalReference"))
+	UPROPERTY(EditAnywhere, BlueprintReadWrite, Category = "Options", meta = (EditCondition = "!bInternalReference", EditConditionHides))
 	FFilePath TargetFile;
 
 	// Use the default prim of the target stage as the referenced/payload prim
-	UPROPERTY(EditAnywhere, BlueprintReadWrite, Category = "Options")
+	UPROPERTY(EditAnywhere, BlueprintReadWrite, Category = "Options", meta = (EditCondition = "!bInternalReference", EditConditionHides))
 	bool bUseDefaultPrim = true;
 
 	// Use a specific prim of the target stage as the referenced/payload prim
-	UPROPERTY(EditAnywhere, BlueprintReadWrite, Category = "Options", meta = (EditCondition = "!bUseDefaultPrim"))
+	UPROPERTY(EditAnywhere, BlueprintReadWrite, Category = "Options", meta = (EditCondition = "!bUseDefaultPrim || bInternalReference"))
 	FString TargetPrimPath;
 
 	// Offset to apply to the referenced/payload prim's time sampled attributes
diff --git a/Engine/Plugins/Runtime/USDCore/Source/USDUtilities/Private/USDConversionUtils.cpp b/Engine/Plugins/Runtime/USDCore/Source/USDUtilities/Private/USDConversionUtils.cpp
--- a/Engine/Plugins/Runtime/USDCore/Source/USDUtilities/Private/USDConversionUtils.cpp
+++ b/Engine/Plugins/Runtime/USDCore/Source/USDUtilities/Private/USDConversionUtils.cpp
@@ -2705,6 +2705,13 @@ namespace UE::UsdConversionUtils::Private
 			return;
 		}
 
+		FString RelativePath = AbsoluteFilePath;
+		if (RelativePath.IsEmpty() && !TargetPrimPath.IsPrimPath())
+		{
+			// Nothing to target
+			return;
+		}
+
 		FScopedUsdAllocs UsdAllocs;
 
 		pxr::UsdPrim UsdPrim(Prim);
@@ -2716,7 +2723,6 @@ namespace UE::UsdConversionUtils::Private
 		}
 
 		// Turn our layer path into a relative one
-		FString RelativePath = AbsoluteFilePath;
 		if (!RelativePath.IsEmpty())
 		{
 			pxr::SdfLayerHandle EditLayer = UsdPrim.GetStage()->GetEditTarget().GetLayer();
