USD: Add warning and fallback for edge case where the onesided reference material is made twosided

diff --git a/Engine/Plugins/Importers/USDImporter/Source/USDSchemas/Private/MeshTranslationImpl.cpp b/Engine/Plugins/Importers/USDImporter/Source/USDSchemas/Private/MeshTranslationImpl.cpp
--- a/Engine/Plugins/Importers/USDImporter/Source/USDSchemas/Private/MeshTranslationImpl.cpp
+++ b/Engine/Plugins/Importers/USDImporter/Source/USDSchemas/Private/MeshTranslationImpl.cpp
@@ -543,6 +543,7 @@ TMap<const UsdUtils::FUsdPrimMaterialSlot*, UMaterialInterface*> MeshTranslation
 					UE::FSdfPath MaterialPrimPath{*Slot.MaterialSource};
 
 					bool bMaterialIsDirectReference = false;
+					bool bUsingTwoSidedPreviewSurfaceAsFallback = false;
 
 					// Here we have to pick the "best" material to use as reference, in case we need compatible/TwoSided versions.
 					// They are returned from PrimLinkCache.GetAssetsForPrim in the most recent to least recent order, so
@@ -574,6 +575,7 @@ TMap<const UsdUtils::FUsdPrimMaterialSlot*, UMaterialInterface*> MeshTranslation
 						if (bSidednessMatches)
 						{
 							Material = ExistingMaterial;
+							bUsingTwoSidedPreviewSurfaceAsFallback = false;
 
 							if (bExistingIsDirectReference)
 							{
@@ -590,17 +592,39 @@ TMap<const UsdUtils::FUsdPrimMaterialSlot*, UMaterialInterface*> MeshTranslation
 							{
 								Material = ExistingMaterial;
 								bMaterialIsDirectReference = bExistingIsDirectReference;
+								bUsingTwoSidedPreviewSurfaceAsFallback = false;
 							}
 						}
-						else	// if (!Slot.bMeshIsDoubleSided && bExistingIsTwoSided)
+						else
 						{
-							// We can ignore this case: If we're searching for a one sided material and just ran into
-							// an existing two-sided one we should just keep iterating: If a two-sided material is within
-							// ExistingMaterials, it's one-sided reference material *must* also be in there, so we'll find
-							// something eventually
+							// If we're searching for a one sided material and just ran into an existing two-sided one we
+							// should keep it just in case, but continue iterating: If a two-sided material is within
+							// ExistingMaterials, its one-sided reference material should also be in there, so we'll find
+							// something eventually. The only exception is in case the user manually set a two-sided material
+							// as a one-sided preview surface parent material, which we'll keep track of and warn about
+							if (!Material)
+							{
+								Material = ExistingMaterial;
+							}
+							if (!Slot.bMeshIsDoubleSided && ExistingMaterial && bExistingIsDirectReference && ExistingMaterial->IsTwoSided())
+							{
+								bUsingTwoSidedPreviewSurfaceAsFallback = true;
+							}
 						}
 					}
 
+					if (bUsingTwoSidedPreviewSurfaceAsFallback)
+					{
+						USD_LOG_USERWARNING(FText::Format(
+							LOCTEXT(
+								"TwoSidedAsOneSidedParent",
+								"Using two-sided material '{0}' for a one-sided mesh prim '{1}' as there is no better alternative. Make sure the ReferencePreviewSurfaceMaterial material set on the project settings is one-sided. You can assign a two-sided material on ReferencePreviewSurfaceTwoSidedMaterial instead."
+							),
+							FText::FromString(Material->GetPathName()),
+							FText::FromString(UsdToUnreal::ConvertPath(UsdPrim.GetPrimPath()))
+						));
+					}
+
 					FString PrefixedMaterialHash = Material ? AssetCache.GetHashForAsset(Material) : FString{};
 					FString HashPrefix = UsdUtils::GetAssetHashPrefix(
 						UsdPrim.GetStage()->GetPrimAtPath(MaterialPrimPath),
diff --git a/Engine/Plugins/Runtime/USDCore/Source/USDClasses/Public/USDProjectSettings.h b/Engine/Plugins/Runtime/USDCore/Source/USDClasses/Public/USDProjectSettings.h
--- a/Engine/Plugins/Runtime/USDCore/Source/USDClasses/Public/USDProjectSettings.h
+++ b/Engine/Plugins/Runtime/USDCore/Source/USDClasses/Public/USDProjectSettings.h
@@ -24,7 +24,7 @@ enum class EUsdEditInInstanceBehavior : uint8
 	ShowPrompt
 };
 
-UENUM(Blueprinttype)
+UENUM(BlueprintType)
 enum class EReferencerTypeHandling : uint8
 {
 	Ignore,
@@ -33,6 +33,23 @@ enum class EReferencerTypeHandling : uint8
 	ShowPrompt,
 };
 
+// We store these here so that we can return our project settings to these default values programmatically if we need to
+namespace UsdDefaultReferenceMaterialPaths
+{
+	const inline FSoftObjectPath ReferencePreviewSurfaceMaterial{TEXT("/USDCore/Materials/UsdPreviewSurface.UsdPreviewSurface")};
+	const inline FSoftObjectPath ReferencePreviewSurfaceTranslucentMaterial{TEXT("/USDCore/Materials/UsdPreviewSurfaceTranslucent.UsdPreviewSurfaceTranslucent")};
+	const inline FSoftObjectPath ReferencePreviewSurfaceTwoSidedMaterial{TEXT("/USDCore/Materials/UsdPreviewSurfaceTwoSided.UsdPreviewSurfaceTwoSided")};
+	const inline FSoftObjectPath ReferencePreviewSurfaceTranslucentTwoSidedMaterial{TEXT("/USDCore/Materials/UsdPreviewSurfaceTranslucentTwoSided.UsdPreviewSurfaceTranslucentTwoSided")};
+	const inline FSoftObjectPath ReferencePreviewSurfaceVTMaterial{TEXT("/USDCore/Materials/UsdPreviewSurfaceVT.UsdPreviewSurfaceVT")};
+	const inline FSoftObjectPath ReferencePreviewSurfaceTranslucentVTMaterial{TEXT("/USDCore/Materials/UsdPreviewSurfaceTranslucentVT.UsdPreviewSurfaceTranslucentVT")};
+	const inline FSoftObjectPath ReferencePreviewSurfaceTwoSidedVTMaterial{TEXT("/USDCore/Materials/UsdPreviewSurfaceTwoSidedVT.UsdPreviewSurfaceTwoSidedVT")};
+	const inline FSoftObjectPath ReferencePreviewSurfaceTranslucentTwoSidedVTMaterial{TEXT("/USDCore/Materials/UsdPreviewSurfaceTranslucentTwoSidedVT.UsdPreviewSurfaceTranslucentTwoSidedVT")};
+	const inline FSoftObjectPath ReferenceDisplayColorMaterial{TEXT("/USDCore/Materials/DisplayColor.DisplayColor")};
+	const inline FSoftObjectPath ReferenceDisplayColorAndOpacityMaterial{TEXT("/USDCore/Materials/DisplayColorAndOpacity.DisplayColorAndOpacity")};
+	const inline FSoftObjectPath ReferenceDisplayColorTwoSidedMaterial{TEXT("/USDCore/Materials/DisplayColorTwoSided.DisplayColorTwoSided")};
+	const inline FSoftObjectPath ReferenceDisplayColorAndOpacityTwoSidedMaterial{TEXT("/USDCore/Materials/DisplayColorAndOpacityTwoSided.DisplayColorAndOpacityTwoSided")};
+}
+
 // USDCore and defaultconfig here so this ends up at DefaultUSDCore.ini in the editor, and is sent to the
 // packaged game as well
 
@@ -175,13 +192,8 @@ public:
 	 * the same names and types as the ones provided by the default material, otherwise the instances will not have
 	 * the parameters filled with values extracted from the USD material when parsing.
 	 */
-	UPROPERTY(
-		config,
-		EditAnywhere,
-		Category = "USD|Reference Materials|UsdPreviewSurface",
-		meta = (AllowedClasses = "/Script/Engine.MaterialInterface")
-	)
-	FSoftObjectPath ReferencePreviewSurfaceMaterial = FSoftObjectPath{TEXT("/USDCore/Materials/UsdPreviewSurface.UsdPreviewSurface")};
+	UPROPERTY(config, EditAnywhere, Category = "USD|Reference Materials|UsdPreviewSurface", meta = (AllowedClasses = "/Script/Engine.MaterialInterface"))
+	FSoftObjectPath ReferencePreviewSurfaceMaterial = UsdDefaultReferenceMaterialPaths::ReferencePreviewSurfaceMaterial;
 
 	/**
 	 * What material to use as reference material when creating material instances from USD materials.
@@ -189,14 +201,8 @@ public:
 	 * the same names and types as the ones provided by the default material, otherwise the instances will not have
 	 * the parameters filled with values extracted from the USD material when parsing.
 	 */
-	UPROPERTY(
-		config,
-		EditAnywhere,
-		Category = "USD|Reference Materials|UsdPreviewSurface",
-		meta = (AllowedClasses = "/Script/Engine.MaterialInterface")
-	)
-	FSoftObjectPath ReferencePreviewSurfaceTranslucentMaterial = FSoftObjectPath{TEXT("/USDCore/Materials/"
-																					  "UsdPreviewSurfaceTranslucent.UsdPreviewSurfaceTranslucent")};
+	UPROPERTY(config, EditAnywhere, Category = "USD|Reference Materials|UsdPreviewSurface", meta = (AllowedClasses = "/Script/Engine.MaterialInterface"))
+	FSoftObjectPath ReferencePreviewSurfaceTranslucentMaterial = UsdDefaultReferenceMaterialPaths::ReferencePreviewSurfaceTranslucentMaterial;
 
 	/**
 	 * What material to use as reference material when creating material instances from USD materials.
@@ -204,14 +210,8 @@ public:
 	 * the same names and types as the ones provided by the default material, otherwise the instances will not have
 	 * the parameters filled with values extracted from the USD material when parsing.
 	 */
-	UPROPERTY(
-		config,
-		EditAnywhere,
-		Category = "USD|Reference Materials|UsdPreviewSurface",
-		meta = (AllowedClasses = "/Script/Engine.MaterialInterface")
-	)
-	FSoftObjectPath ReferencePreviewSurfaceTwoSidedMaterial = FSoftObjectPath{TEXT("/USDCore/Materials/"
-																				   "UsdPreviewSurfaceTwoSided.UsdPreviewSurfaceTwoSided")};
+	UPROPERTY(config, EditAnywhere, Category = "USD|Reference Materials|UsdPreviewSurface", meta = (AllowedClasses = "/Script/Engine.MaterialInterface"))
+	FSoftObjectPath ReferencePreviewSurfaceTwoSidedMaterial = UsdDefaultReferenceMaterialPaths::ReferencePreviewSurfaceTwoSidedMaterial;
 
 	/**
 	 * What material to use as reference material when creating material instances from USD materials.
@@ -219,15 +219,8 @@ public:
 	 * the same names and types as the ones provided by the default material, otherwise the instances will not have
 	 * the parameters filled with values extracted from the USD material when parsing.
 	 */
-	UPROPERTY(
-		config,
-		EditAnywhere,
-		Category = "USD|Reference Materials|UsdPreviewSurface",
-		meta = (AllowedClasses = "/Script/Engine.MaterialInterface")
-	)
-	FSoftObjectPath ReferencePreviewSurfaceTranslucentTwoSidedMaterial = FSoftObjectPath{
-		TEXT("/USDCore/Materials/UsdPreviewSurfaceTranslucentTwoSided.UsdPreviewSurfaceTranslucentTwoSided")
-	};
+	UPROPERTY(config, EditAnywhere, Category = "USD|Reference Materials|UsdPreviewSurface", meta = (AllowedClasses = "/Script/Engine.MaterialInterface"))
+	FSoftObjectPath ReferencePreviewSurfaceTranslucentTwoSidedMaterial = UsdDefaultReferenceMaterialPaths::ReferencePreviewSurfaceTranslucentTwoSidedMaterial;
 
 	/**
 	 * What material to use as reference material when creating material instances from USD materials.
@@ -235,13 +228,8 @@ public:
 	 * the same names and types as the ones provided by the default material, otherwise the instances will not have
 	 * the parameters filled with values extracted from the USD material when parsing.
 	 */
-	UPROPERTY(
-		config,
-		EditAnywhere,
-		Category = "USD|Reference Materials|UsdPreviewSurface with Virtual Textures",
-		meta = (AllowedClasses = "/Script/Engine.MaterialInterface")
-	)
-	FSoftObjectPath ReferencePreviewSurfaceVTMaterial = FSoftObjectPath{TEXT("/USDCore/Materials/UsdPreviewSurfaceVT.UsdPreviewSurfaceVT")};
+	UPROPERTY(config, EditAnywhere, Category = "USD|Reference Materials|UsdPreviewSurface with Virtual Textures", meta = (AllowedClasses = "/Script/Engine.MaterialInterface"))
+	FSoftObjectPath ReferencePreviewSurfaceVTMaterial = UsdDefaultReferenceMaterialPaths::ReferencePreviewSurfaceVTMaterial;
 
 	/**
 	 * What material to use as reference material when creating material instances from USD materials.
@@ -249,15 +237,8 @@ public:
 	 * the same names and types as the ones provided by the default material, otherwise the instances will not have
 	 * the parameters filled with values extracted from the USD material when parsing.
 	 */
-	UPROPERTY(
-		config,
-		EditAnywhere,
-		Category = "USD|Reference Materials|UsdPreviewSurface with Virtual Textures",
-		meta = (AllowedClasses = "/Script/Engine.MaterialInterface")
-	)
-	FSoftObjectPath ReferencePreviewSurfaceTranslucentVTMaterial = FSoftObjectPath{
-		TEXT("/USDCore/Materials/UsdPreviewSurfaceTranslucentVT.UsdPreviewSurfaceTranslucentVT")
-	};
+	UPROPERTY(config, EditAnywhere, Category = "USD|Reference Materials|UsdPreviewSurface with Virtual Textures", meta = (AllowedClasses = "/Script/Engine.MaterialInterface"))
+	FSoftObjectPath ReferencePreviewSurfaceTranslucentVTMaterial = UsdDefaultReferenceMaterialPaths::ReferencePreviewSurfaceTranslucentVTMaterial;
 
 	/**
 	 * What material to use as reference material when creating material instances from USD materials.
@@ -265,14 +246,8 @@ public:
 	 * the same names and types as the ones provided by the default material, otherwise the instances will not have
 	 * the parameters filled with values extracted from the USD material when parsing.
 	 */
-	UPROPERTY(
-		config,
-		EditAnywhere,
-		Category = "USD|Reference Materials|UsdPreviewSurface with Virtual Textures",
-		meta = (AllowedClasses = "/Script/Engine.MaterialInterface")
-	)
-	FSoftObjectPath ReferencePreviewSurfaceTwoSidedVTMaterial = FSoftObjectPath{TEXT("/USDCore/Materials/"
-																					 "UsdPreviewSurfaceTwoSidedVT.UsdPreviewSurfaceTwoSidedVT")};
+	UPROPERTY(config, EditAnywhere, Category = "USD|Reference Materials|UsdPreviewSurface with Virtual Textures", meta = (AllowedClasses = "/Script/Engine.MaterialInterface"))
+	FSoftObjectPath ReferencePreviewSurfaceTwoSidedVTMaterial = UsdDefaultReferenceMaterialPaths::ReferencePreviewSurfaceTwoSidedVTMaterial;
 
 	/**
 	 * What material to use as reference material when creating material instances from USD materials.
@@ -280,15 +255,8 @@ public:
 	 * the same names and types as the ones provided by the default material, otherwise the instances will not have
 	 * the parameters filled with values extracted from the USD material when parsing.
 	 */
-	UPROPERTY(
-		config,
-		EditAnywhere,
-		Category = "USD|Reference Materials|UsdPreviewSurface with Virtual Textures",
-		meta = (AllowedClasses = "/Script/Engine.MaterialInterface")
-	)
-	FSoftObjectPath ReferencePreviewSurfaceTranslucentTwoSidedVTMaterial = FSoftObjectPath{
-		TEXT("/USDCore/Materials/UsdPreviewSurfaceTranslucentTwoSidedVT.UsdPreviewSurfaceTranslucentTwoSidedVT")
-	};
+	UPROPERTY(config, EditAnywhere, Category = "USD|Reference Materials|UsdPreviewSurface with Virtual Textures", meta = (AllowedClasses = "/Script/Engine.MaterialInterface"))
+	FSoftObjectPath ReferencePreviewSurfaceTranslucentTwoSidedVTMaterial = UsdDefaultReferenceMaterialPaths::ReferencePreviewSurfaceTranslucentTwoSidedVTMaterial;
 
 	/**
 	 * What material to use as reference material when creating material instances from USD materials.
@@ -297,7 +265,7 @@ public:
 	 * the parameters filled with values extracted from the USD material when parsing.
 	 */
 	UPROPERTY(config, EditAnywhere, Category = "USD|Reference Materials|DisplayColor", meta = (AllowedClasses = "/Script/Engine.MaterialInterface"))
-	FSoftObjectPath ReferenceDisplayColorMaterial = FSoftObjectPath{TEXT("/USDCore/Materials/DisplayColor.DisplayColor")};
+	FSoftObjectPath ReferenceDisplayColorMaterial = UsdDefaultReferenceMaterialPaths::ReferenceDisplayColorMaterial;
 
 	/**
 	 * What material to use as reference material when creating material instances from USD materials.
@@ -306,8 +274,7 @@ public:
 	 * the parameters filled with values extracted from the USD material when parsing.
 	 */
 	UPROPERTY(config, EditAnywhere, Category = "USD|Reference Materials|DisplayColor", meta = (AllowedClasses = "/Script/Engine.MaterialInterface"))
-	FSoftObjectPath ReferenceDisplayColorAndOpacityMaterial = FSoftObjectPath{TEXT("/USDCore/Materials/"
-																				   "DisplayColorAndOpacity.DisplayColorAndOpacity")};
+	FSoftObjectPath ReferenceDisplayColorAndOpacityMaterial = UsdDefaultReferenceMaterialPaths::ReferenceDisplayColorAndOpacityMaterial;
 
 	/**
 	 * What material to use as reference material when creating material instances from USD materials.
@@ -316,7 +283,7 @@ public:
 	 * the parameters filled with values extracted from the USD material when parsing.
 	 */
 	UPROPERTY(config, EditAnywhere, Category = "USD|Reference Materials|DisplayColor", meta = (AllowedClasses = "/Script/Engine.MaterialInterface"))
-	FSoftObjectPath ReferenceDisplayColorTwoSidedMaterial = FSoftObjectPath{TEXT("/USDCore/Materials/DisplayColorTwoSided.DisplayColorTwoSided")};
+	FSoftObjectPath ReferenceDisplayColorTwoSidedMaterial = UsdDefaultReferenceMaterialPaths::ReferenceDisplayColorTwoSidedMaterial;
 
 	/**
 	 * What material to use as reference material when creating material instances from USD materials.
@@ -325,9 +292,7 @@ public:
 	 * the parameters filled with values extracted from the USD material when parsing.
 	 */
 	UPROPERTY(config, EditAnywhere, Category = "USD|Reference Materials|DisplayColor", meta = (AllowedClasses = "/Script/Engine.MaterialInterface"))
-	FSoftObjectPath ReferenceDisplayColorAndOpacityTwoSidedMaterial = FSoftObjectPath{
-		TEXT("/USDCore/Materials/DisplayColorAndOpacityTwoSided.DisplayColorAndOpacityTwoSided")
-	};
+	FSoftObjectPath ReferenceDisplayColorAndOpacityTwoSidedMaterial = UsdDefaultReferenceMaterialPaths::ReferenceDisplayColorAndOpacityTwoSidedMaterial;
 
 public:
 #if WITH_EDITOR
diff --git a/Engine/Plugins/Runtime/USDCore/Source/USDUtilities/Private/USDUtilitiesModule.cpp b/Engine/Plugins/Runtime/USDCore/Source/USDUtilities/Private/USDUtilitiesModule.cpp
--- a/Engine/Plugins/Runtime/USDCore/Source/USDUtilities/Private/USDUtilitiesModule.cpp
+++ b/Engine/Plugins/Runtime/USDCore/Source/USDUtilities/Private/USDUtilitiesModule.cpp
@@ -4,13 +4,17 @@
 
 #include "USDErrorUtils.h"
 #include "USDMemory.h"
+#include "USDTypesConversion.h"
 
+#include "Materials/MaterialInterface.h"
 #include "Modules/ModuleManager.h"
 
 #if WITH_EDITOR
 #include "MessageLogModule.h"
 #endif	  // WITH_EDITOR
 
+#define LOCTEXT_NAMESPACE "UsdUtilitiesModule"
+
 class FUsdUtilitiesModule : public IUsdUtilitiesModule
 {
 	FDelegateHandle OnProjectSettingsChangedHandle;
@@ -47,6 +51,68 @@ public:
 						FUsdLogManager::UnregisterDiagnosticDelegate();
 					}
 				}
+
+			// Produce some immediate warning in case the user adds a an invalid material for one of the reference material properties
+#define CHECK_MATERIAL(ReferenceType, bShouldBeTwoSided, bShouldBeTranslucent)                                                                                     				\
+				if (PropertyChangedEvent.GetPropertyName() == GET_MEMBER_NAME_CHECKED(UUsdProjectSettings, ReferenceType))                                                      \
+				{                                                                                                                                                               \
+					if (UMaterialInterface* Material = Cast<UMaterialInterface>(UsdSettings->ReferenceType.TryLoad()))                                                          \
+					{                                                                                                                                                           \
+						FScopedUsdMessageLog ScopedMessageLog;                                                                                                                  \
+																																												\
+						bool bRevertMaterial = false;                                                                                                                           \
+																																												\
+						const bool bIsTranslucent = Material->GetBlendMode() == EBlendMode::BLEND_Translucent                                                                   \
+													|| Material->GetBlendMode() == EBlendMode::BLEND_Additive                                                                   \
+													|| Material->GetBlendMode() == EBlendMode::BLEND_Modulate;                                                                  \
+						if (bIsTranslucent != bShouldBeTranslucent)                                                                                                             \
+						{                                                                                                                                                       \
+							USD_LOG_USERWARNING(FText::Format(                                                                                                                  \
+								LOCTEXT(                                                                                                                                        \
+									"TranslucentMaterialMismatch",                                                                                                              \
+									"Material '{0}' {1} but the property {2}! Make sure to assign translucent materials only to the translucent reference material properties!" \
+								),                                                                                                                                              \
+								FText::FromString(Material->GetPathName()),                                                                                                     \
+								FText::FromString(bIsTranslucent ? TEXT("is translucent") : TEXT("is not translucent")),                                                        \
+								FText::FromString(bShouldBeTranslucent ? TEXT("is translucent") : TEXT("is not translucent"))                                                   \
+							));                                                                                                                                                 \
+							bRevertMaterial = true;                                                                                                                             \
+						}                                                                                                                                                       \
+																																												\
+						if (Material->IsTwoSided() != bShouldBeTwoSided)                                                                                                        \
+						{                                                                                                                                                       \
+							USD_LOG_USERWARNING(FText::Format(                                                                                                                  \
+								LOCTEXT(                                                                                                                                        \
+									"TwoSidedMaterialMismatch",                                                                                                                 \
+									"Material '{0}' {1} but the property {2}! Make sure to assign two-sided materials only to the two-sided reference material properties!"     \
+								),                                                                                                                                              \
+								FText::FromString(Material->GetPathName()),                                                                                                     \
+								FText::FromString(Material->IsTwoSided() ? TEXT("is two-sided") : TEXT("is one-sided")),                                                        \
+								FText::FromString(bShouldBeTwoSided ? TEXT("is two-sided") : TEXT("is one-sided"))                                                              \
+							));                                                                                                                                                 \
+							bRevertMaterial = true;                                                                                                                             \
+						}                                                                                                                                                       \
+																																												\
+						if (bRevertMaterial)                                                                                                                                    \
+						{                                                                                                                                                       \
+							UsdSettings->ReferenceType = UsdDefaultReferenceMaterialPaths::ReferenceType;                                                                       \
+							UsdSettings->SaveConfig();                                                                                                                          \
+						}                                                                                                                                                       \
+					}                                                                                                                                                           \
+				}
+				CHECK_MATERIAL(ReferencePreviewSurfaceMaterial, false, false);
+				CHECK_MATERIAL(ReferencePreviewSurfaceTranslucentMaterial, false, true);
+				CHECK_MATERIAL(ReferencePreviewSurfaceTwoSidedMaterial, true, false);
+				CHECK_MATERIAL(ReferencePreviewSurfaceTranslucentTwoSidedMaterial, true, true);
+				CHECK_MATERIAL(ReferencePreviewSurfaceVTMaterial, false, false);
+				CHECK_MATERIAL(ReferencePreviewSurfaceTranslucentVTMaterial, false, true);
+				CHECK_MATERIAL(ReferencePreviewSurfaceTwoSidedVTMaterial, true, false);
+				CHECK_MATERIAL(ReferencePreviewSurfaceTranslucentTwoSidedVTMaterial, true, true);
+				CHECK_MATERIAL(ReferenceDisplayColorMaterial, false, false);
+				CHECK_MATERIAL(ReferenceDisplayColorAndOpacityMaterial, false, true);
+				CHECK_MATERIAL(ReferenceDisplayColorTwoSidedMaterial, true, false);
+				CHECK_MATERIAL(ReferenceDisplayColorAndOpacityTwoSidedMaterial, true, true);
+#undef CHECK_MATERIAL
 			}
 		);
 
@@ -78,3 +144,5 @@ public:
 };
 
 IMPLEMENT_MODULE_USD(FUsdUtilitiesModule, USDUtilities);
+
+#undef LOCTEXT_NAMESPACE
\ No newline at end of file
