USDCore: remove use of deprecated Usd-level file format headers

These headers were removed in OpenUSD version 25.11. Their
Sdf-level counterparts should be used instead and were available in
our current 25.08 version, so no preprocessor guards are necessary.

#rnx

diff --git a/Engine/Plugins/Experimental/ChaosCachingUSD/Source/ChaosCachingUSD/Private/ChaosCachingUSD/Util.h b/Engine/Plugins/Experimental/ChaosCachingUSD/Source/ChaosCachingUSD/Private/ChaosCachingUSD/Util.h
--- a/Engine/Plugins/Experimental/ChaosCachingUSD/Source/ChaosCachingUSD/Private/ChaosCachingUSD/Util.h
+++ b/Engine/Plugins/Experimental/ChaosCachingUSD/Source/ChaosCachingUSD/Private/ChaosCachingUSD/Util.h
@@ -4,9 +4,9 @@
 
 #include "USDMemory.h"
 #include "USDIncludesStart.h"
-	#include "pxr/usd/usd/usdaFileFormat.h"
-	#include "pxr/usd/usd/usdcFileFormat.h"
-	#include "pxr/usd/usd/usdFileFormat.h"
+	#include "pxr/usd/sdf/usdFileFormat.h"
+	#include "pxr/usd/sdf/usdaFileFormat.h"
+	#include "pxr/usd/sdf/usdcFileFormat.h"
 
 	#include "pxr/usd/usd/tokens.h"
 	#include "pxr/usd/usdUtils/stitchClips.h"
diff --git a/Engine/Plugins/Runtime/USDCore/Source/USDUtilities/Private/USDLayerUtils.cpp b/Engine/Plugins/Runtime/USDCore/Source/USDUtilities/Private/USDLayerUtils.cpp
--- a/Engine/Plugins/Runtime/USDCore/Source/USDUtilities/Private/USDLayerUtils.cpp
+++ b/Engine/Plugins/Runtime/USDCore/Source/USDUtilities/Private/USDLayerUtils.cpp
@@ -28,12 +28,12 @@
 #include "pxr/usd/sdf/fileFormat.h"
 #include "pxr/usd/sdf/layer.h"
 #include "pxr/usd/sdf/layerUtils.h"
-#include "pxr/usd/sdf/textFileFormat.h"
+#include "pxr/usd/sdf/usdaFileFormat.h"
+#include "pxr/usd/sdf/zipFile.h"
 #include "pxr/usd/usd/editContext.h"
 #include "pxr/usd/usd/modelAPI.h"
 #include "pxr/usd/usd/primCompositionQuery.h"
 #include "pxr/usd/usd/stage.h"
-#include "pxr/usd/usd/zipFile.h"
 #include "pxr/usd/usdGeom/xform.h"
 #include "pxr/usd/usdUtils/flattenLayerStack.h"
 #include "USDIncludesEnd.h"
@@ -365,10 +365,10 @@ bool UsdUtils::DecompressUSDZFile(const FString& InUSDZFilePath, const FString&
 
 	bool bTraversedFirstFile = false;
 
-	pxr::UsdZipFile File = pxr::UsdZipFile::Open(UnrealToUsd::ConvertString(*InUSDZFilePath).Get());
-	for (pxr::UsdZipFile::Iterator Iter = File.begin(), End = File.end(); Iter != End; ++Iter)
+	pxr::SdfZipFile File = pxr::SdfZipFile::Open(UnrealToUsd::ConvertString(*InUSDZFilePath).Get());
+	for (pxr::SdfZipFile::Iterator Iter = File.begin(), End = File.end(); Iter != End; ++Iter)
 	{
-		const pxr::UsdZipFile::FileInfo Info = Iter.GetFileInfo();
+		const pxr::SdfZipFile::FileInfo Info = Iter.GetFileInfo();
 		const std::string InnerFilePath = *Iter;
 		const FString UEInnerFilePath = UsdToUnreal::ConvertString(InnerFilePath);
 		const char* FileStart = Iter.GetFile();
@@ -774,7 +774,7 @@ UE::FSdfLayer UsdUtils::GetUEPersistentStateSublayer(const UE::FUsdStage& Stage,
 	if (!StateLayer && bCreateIfNeeded)
 	{
 		StateLayer = pxr::SdfLayer::New(
-			pxr::SdfFileFormat::FindById(pxr::SdfTextFileFormatTokens->Id),
+			pxr::SdfFileFormat::FindById(pxr::SdfUsdaFileFormatTokens->Id),
 			UnrealToUsd::ConvertString(*ExpectedStateLayerPath).Get()
 		);
 	}
diff --git a/Engine/Plugins/Runtime/USDCore/Source/UnrealUSDWrapper/Private/UnrealUSDWrapper.cpp b/Engine/Plugins/Runtime/USDCore/Source/UnrealUSDWrapper/Private/UnrealUSDWrapper.cpp
--- a/Engine/Plugins/Runtime/USDCore/Source/UnrealUSDWrapper/Private/UnrealUSDWrapper.cpp
+++ b/Engine/Plugins/Runtime/USDCore/Source/UnrealUSDWrapper/Private/UnrealUSDWrapper.cpp
@@ -51,6 +51,10 @@
 #include "pxr/usd/sdf/fileFormat.h"
 #include "pxr/usd/sdf/path.h"
 #include "pxr/usd/sdf/schema.h"
+#include "pxr/usd/sdf/usdFileFormat.h"
+#include "pxr/usd/sdf/usdaFileFormat.h"
+#include "pxr/usd/sdf/usdcFileFormat.h"
+#include "pxr/usd/sdf/usdzFileFormat.h"
 #include "pxr/usd/usd/attribute.h"
 #include "pxr/usd/usd/common.h"
 #include "pxr/usd/usd/debugCodes.h"
@@ -60,10 +64,6 @@
 #include "pxr/usd/usd/schemaBase.h"
 #include "pxr/usd/usd/stage.h"
 #include "pxr/usd/usd/stageCacheContext.h"
-#include "pxr/usd/usd/usdaFileFormat.h"
-#include "pxr/usd/usd/usdcFileFormat.h"
-#include "pxr/usd/usd/usdFileFormat.h"
-#include "pxr/usd/usd/usdzFileFormat.h"
 #include "pxr/usd/usd/variantSets.h"
 #include "pxr/usd/usdGeom/mesh.h"
 #include "pxr/usd/usdGeom/metrics.h"
@@ -698,7 +698,7 @@ TArray<FString> UnrealUSDWrapper::GetAllSupportedFileFormats() TSAN_SAFE
 	for (const std::string& Ext : Extensions)
 	{
 		// Ignore formats that don't target "usd"
-		pxr::SdfFileFormatConstPtr Format = pxr::SdfFileFormat::FindByExtension(Ext, pxr::UsdUsdFileFormatTokens->Target);
+		pxr::SdfFileFormatConstPtr Format = pxr::SdfFileFormat::FindByExtension(Ext, pxr::SdfUsdFileFormatTokens->Target);
 		if (Format == nullptr)
 		{
 			continue;
@@ -721,10 +721,10 @@ TArray<FString> UnrealUSDWrapper::GetNativeFileFormats() TSAN_SAFE
 	FScopedUsdAllocs Allocs;
 
 	const pxr::TfTokenVector NativeFormatIds{
-		pxr::UsdUsdFileFormatTokens->Id,
-		pxr::UsdUsdaFileFormatTokens->Id,
-		pxr::UsdUsdcFileFormatTokens->Id,
-		pxr::UsdUsdzFileFormatTokens->Id
+		pxr::SdfUsdFileFormatTokens->Id,
+		pxr::SdfUsdaFileFormatTokens->Id,
+		pxr::SdfUsdcFileFormatTokens->Id,
+		pxr::SdfUsdzFileFormatTokens->Id
 	};
 
 	std::set<std::string> FileExtensions;
@@ -760,7 +760,7 @@ void UnrealUSDWrapper::GetNativeFileFormats(TArray<FString>& OutTextFormats, TAr
 	FScopedUsdAllocs Allocs;
 
 	std::set<std::string> TextExtensions;
-	const pxr::SdfFileFormatConstPtr TextFormat = pxr::SdfFileFormat::FindById(pxr::UsdUsdaFileFormatTokens->Id);
+	const pxr::SdfFileFormatConstPtr TextFormat = pxr::SdfFileFormat::FindById(pxr::SdfUsdaFileFormatTokens->Id);
 	if (TextFormat)
 	{
 		for (const std::string& FileExtension : TextFormat->GetFileExtensions())
@@ -776,9 +776,9 @@ void UnrealUSDWrapper::GetNativeFileFormats(TArray<FString>& OutTextFormats, TAr
 	}
 
 	const static std::vector<pxr::TfToken> BinaryFormatIds{
-		pxr::UsdUsdFileFormatTokens->Id,
-		pxr::UsdUsdcFileFormatTokens->Id,
-		pxr::UsdUsdzFileFormatTokens->Id
+		pxr::SdfUsdFileFormatTokens->Id,
+		pxr::SdfUsdcFileFormatTokens->Id,
+		pxr::SdfUsdzFileFormatTokens->Id
 	};
 
 	std::set<std::string> PossiblyBinaryExtensions;
