USD: Fixed message about loading the UnrealUSDWrapper module before using its functions. The message was triggered by the creation of the CDO for UUsdInfoCache so this change just explicitly loads the UnrealUSDWrapper module when that happens (it loaded it implicitly anyway when the message was triggered).

diff --git a/Engine/Plugins/Runtime/USDCore/Source/USDUtilities/Private/Objects/USDInfoCache.cpp b/Engine/Plugins/Runtime/USDCore/Source/USDUtilities/Private/Objects/USDInfoCache.cpp
--- a/Engine/Plugins/Runtime/USDCore/Source/USDUtilities/Private/Objects/USDInfoCache.cpp
+++ b/Engine/Plugins/Runtime/USDCore/Source/USDUtilities/Private/Objects/USDInfoCache.cpp
@@ -522,6 +522,14 @@ public:
 
 FUsdInfoCache::FUsdInfoCache()
 {
+	// Make sure the UnrealUSDWrapper module is loaded first since FUsdInfoCacheImpl will call UnrealUSDWrapper::GetNativeFileFormats
+	const FName ModuleName = TEXT("UnrealUSDWrapper");
+	FModuleManager& Manager = FModuleManager::Get();
+	if (!Manager.IsModuleLoaded(ModuleName))
+	{
+		Manager.LoadModule(ModuleName);
+	}
+
 	Impl = MakeUnique<FUsdInfoCacheImpl>();
 }
 
