USD: Fix localization issues due to using LOCTEXT inside another macro

#Synced-CL 47704025

diff --git a/Engine/Plugins/Runtime/USDCore/Source/USDUtilities/Private/USDUtilitiesModule.cpp b/Engine/Plugins/Runtime/USDCore/Source/USDUtilities/Private/USDUtilitiesModule.cpp
--- a/Engine/Plugins/Runtime/USDCore/Source/USDUtilities/Private/USDUtilitiesModule.cpp
+++ b/Engine/Plugins/Runtime/USDCore/Source/USDUtilities/Private/USDUtilitiesModule.cpp
@@ -31,6 +31,21 @@ public:
 #if WITH_EDITOR
 		LLM_SCOPE_BYTAG(Usd);
 
+		// These strings are used for the CHECK_MATERIAL macro, but aren't inline because the localization system
+		// doesn't know what to do with the LOCTEXT inside the other macro
+		const static FText TranslucentMismatchMessage = LOCTEXT(
+			"TranslucentMaterialMismatch",
+			"Material '{0}' {1} but the property {2}! Make sure to assign translucent materials only to the translucent reference material properties!"
+		);
+		const static FText TwoSidedMistmachMessage = LOCTEXT(
+			"TwoSidedMaterialMismatch",
+			"Material '{0}' {1} but the property {2}! Make sure to assign two-sided materials only to the two-sided reference material properties!"
+		);
+		const static FText IsTwoSidedText = LOCTEXT("IsTwoSidedText", "is two-sided");
+		const static FText IsOneSidedText = LOCTEXT("IsOneSidedText", "is one-sided");
+		const static FText IsTranslucentText = LOCTEXT("IsTranslucentText", "is translucent");
+		const static FText IsNotTranslucentText = LOCTEXT("IsNotTranslucentText", "is not translucent");
+
 		OnProjectSettingsChangedHandle = ProjectSettings->OnSettingChanged().AddLambda(
 			[](UObject* SettingsObject, struct FPropertyChangedEvent& PropertyChangedEvent)
 			{
@@ -53,52 +68,46 @@ public:
 				}
 
 			// Produce some immediate warning in case the user adds a an invalid material for one of the reference material properties
-#define CHECK_MATERIAL(ReferenceType, bShouldBeTwoSided, bShouldBeTranslucent)                                                                                     				\
-				if (PropertyChangedEvent.GetPropertyName() == GET_MEMBER_NAME_CHECKED(UUsdProjectSettings, ReferenceType))                                                      \
-				{                                                                                                                                                               \
-					if (UMaterialInterface* Material = Cast<UMaterialInterface>(UsdSettings->ReferenceType.TryLoad()))                                                          \
-					{                                                                                                                                                           \
-						FScopedUsdMessageLog ScopedMessageLog;                                                                                                                  \
-																																												\
-						bool bRevertMaterial = false;                                                                                                                           \
-																																												\
-						const bool bIsTranslucent = Material->GetBlendMode() == EBlendMode::BLEND_Translucent                                                                   \
-													|| Material->GetBlendMode() == EBlendMode::BLEND_Additive                                                                   \
-													|| Material->GetBlendMode() == EBlendMode::BLEND_Modulate;                                                                  \
-						if (bIsTranslucent != bShouldBeTranslucent)                                                                                                             \
-						{                                                                                                                                                       \
-							USD_LOG_USERWARNING(FText::Format(                                                                                                                  \
-								LOCTEXT(                                                                                                                                        \
-									"TranslucentMaterialMismatch",                                                                                                              \
-									"Material '{0}' {1} but the property {2}! Make sure to assign translucent materials only to the translucent reference material properties!" \
-								),                                                                                                                                              \
-								FText::FromString(Material->GetPathName()),                                                                                                     \
-								FText::FromString(bIsTranslucent ? TEXT("is translucent") : TEXT("is not translucent")),                                                        \
-								FText::FromString(bShouldBeTranslucent ? TEXT("is translucent") : TEXT("is not translucent"))                                                   \
-							));                                                                                                                                                 \
-							bRevertMaterial = true;                                                                                                                             \
-						}                                                                                                                                                       \
-																																												\
-						if (Material->IsTwoSided() != bShouldBeTwoSided)                                                                                                        \
-						{                                                                                                                                                       \
-							USD_LOG_USERWARNING(FText::Format(                                                                                                                  \
-								LOCTEXT(                                                                                                                                        \
-									"TwoSidedMaterialMismatch",                                                                                                                 \
-									"Material '{0}' {1} but the property {2}! Make sure to assign two-sided materials only to the two-sided reference material properties!"     \
-								),                                                                                                                                              \
-								FText::FromString(Material->GetPathName()),                                                                                                     \
-								FText::FromString(Material->IsTwoSided() ? TEXT("is two-sided") : TEXT("is one-sided")),                                                        \
-								FText::FromString(bShouldBeTwoSided ? TEXT("is two-sided") : TEXT("is one-sided"))                                                              \
-							));                                                                                                                                                 \
-							bRevertMaterial = true;                                                                                                                             \
-						}                                                                                                                                                       \
-																																												\
-						if (bRevertMaterial)                                                                                                                                    \
-						{                                                                                                                                                       \
-							UsdSettings->ReferenceType = UsdDefaultReferenceMaterialPaths::ReferenceType;                                                                       \
-							UsdSettings->SaveConfig();                                                                                                                          \
-						}                                                                                                                                                       \
-					}                                                                                                                                                           \
+#define CHECK_MATERIAL(ReferenceType, bShouldBeTwoSided, bShouldBeTranslucent)                                      	   \
+				if (PropertyChangedEvent.GetPropertyName() == GET_MEMBER_NAME_CHECKED(UUsdProjectSettings, ReferenceType)) \
+				{                                                                                                          \
+					if (UMaterialInterface* Material = Cast<UMaterialInterface>(UsdSettings->ReferenceType.TryLoad()))     \
+					{                                                                                                      \
+						FScopedUsdMessageLog ScopedMessageLog;                                                             \
+																														   \
+						bool bRevertMaterial = false;                                                                      \
+																														   \
+						const bool bIsTranslucent = Material->GetBlendMode() == EBlendMode::BLEND_Translucent              \
+													|| Material->GetBlendMode() == EBlendMode::BLEND_Additive              \
+													|| Material->GetBlendMode() == EBlendMode::BLEND_Modulate;             \
+						if (bIsTranslucent != bShouldBeTranslucent)                                                        \
+						{                                                                                                  \
+							USD_LOG_USERWARNING(FText::Format(                                                             \
+								TranslucentMismatchMessage,                                                                \
+								FText::FromString(Material->GetPathName()),                                                \
+								bIsTranslucent ? IsTranslucentText : IsNotTranslucentText,                                 \
+								bShouldBeTranslucent ? IsTranslucentText : IsNotTranslucentText                            \
+							));                                                                                            \
+							bRevertMaterial = true;                                                                        \
+						}                                                                                                  \
+																														   \
+						if (Material->IsTwoSided() != bShouldBeTwoSided)                                                   \
+						{                                                                                                  \
+							USD_LOG_USERWARNING(FText::Format(                                                             \
+								TwoSidedMistmachMessage,                                                                   \
+								FText::FromString(Material->GetPathName()),                                                \
+								Material->IsTwoSided() ? IsTwoSidedText : IsOneSidedText,                                  \
+								bShouldBeTwoSided ? IsTwoSidedText : IsOneSidedText                                        \
+							));                                                                                            \
+							bRevertMaterial = true;                                                                        \
+						}                                                                                                  \
+																														   \
+						if (bRevertMaterial)                                                                               \
+						{                                                                                                  \
+							UsdSettings->ReferenceType = UsdDefaultReferenceMaterialPaths::ReferenceType;                  \
+							UsdSettings->SaveConfig();                                                                     \
+						}                                                                                                  \
+					}                                                                                                      \
 				}
 				CHECK_MATERIAL(ReferencePreviewSurfaceMaterial, false, false);
 				CHECK_MATERIAL(ReferencePreviewSurfaceTranslucentMaterial, false, true);
