USD: Fixed level cannot be saved after doing a level import from USD with metadata collected on components and converting an actor to Blueprint.

#rnx
#Synced-CL 48415188

diff --git a/Engine/Plugins/Runtime/USDCore/Source/USDClasses/Private/USDObjectUtils.cpp b/Engine/Plugins/Runtime/USDCore/Source/USDClasses/Private/USDObjectUtils.cpp
--- a/Engine/Plugins/Runtime/USDCore/Source/USDClasses/Private/USDObjectUtils.cpp
+++ b/Engine/Plugins/Runtime/USDCore/Source/USDClasses/Private/USDObjectUtils.cpp
@@ -205,7 +205,7 @@ UUsdAssetUserData* UsdUnreal::ObjectUtils::GetOrCreateAssetUserData(UObject* Obj
 		// For now we're expecting objects to only have one instance of UUsdAssetUserData
 		ensure(!AssetUserDataInterface->HasAssetUserDataOfClass(UUsdAssetUserData::StaticClass()));
 
-		AssetUserData = NewObject<UUsdAssetUserData>(Object, Class, TEXT("UsdAssetUserData"));
+		AssetUserData = NewObject<UUsdAssetUserData>(Object, Class, TEXT("UsdAssetUserData"), RF_Public);
 		AssetUserDataInterface->AddAssetUserData(AssetUserData);
 	}
 
diff --git a/Engine/Plugins/Runtime/USDCore/Source/USDUtilities/Private/USDConversionUtils.cpp b/Engine/Plugins/Runtime/USDCore/Source/USDUtilities/Private/USDConversionUtils.cpp
--- a/Engine/Plugins/Runtime/USDCore/Source/USDUtilities/Private/USDConversionUtils.cpp
+++ b/Engine/Plugins/Runtime/USDCore/Source/USDUtilities/Private/USDConversionUtils.cpp
@@ -2498,7 +2498,7 @@ UUsdAssetUserData* UsdUtils::GetOrCreateAssetUserData(UObject* Object, TSubclass
 		// For now we're expecting objects to only have one instance of UUsdAssetUserData
 		ensure(!AssetUserDataInterface->HasAssetUserDataOfClass(UUsdAssetUserData::StaticClass()));
 
-		AssetUserData = NewObject<UUsdAssetUserData>(Object, Class, TEXT("UsdAssetUserData"));
+		AssetUserData = NewObject<UUsdAssetUserData>(Object, Class, TEXT("UsdAssetUserData"), RF_Public);
 		AssetUserDataInterface->AddAssetUserData(AssetUserData);
 	}
 
