USD Interchange: Analytics updates.

#Synced-CL 49553638

diff --git a/Engine/Plugins/Importers/USDImporter/Source/USDSchemas/Private/USDSchemasModule.cpp b/Engine/Plugins/Importers/USDImporter/Source/USDSchemas/Private/USDSchemasModule.cpp
--- a/Engine/Plugins/Importers/USDImporter/Source/USDSchemas/Private/USDSchemasModule.cpp
+++ b/Engine/Plugins/Importers/USDImporter/Source/USDSchemas/Private/USDSchemasModule.cpp
@@ -121,271 +121,20 @@ protected:
 void UsdUnreal::Analytics::CollectSchemaAnalytics(const UE::FUsdStage& Stage, const FString& EventName)
 {
 #if USE_USD_SDK
-	if (!Stage)
-	{
-		return;
-	}
-
-	// Last collected at cl 32482283
-	const static TSet<FString> NativeSchemaNames{
-		"AssetPreviewsAPI",
-		"Backdrop",
-		"BasisCurves",
-		"BlendShape",
-		"Camera",
-		"Capsule",
-		"Capsule_1",
-		"ClipsAPI",
-		"CollectionAPI",
-		"Cone",
-		"ConnectableAPI",
-		"ControlRigAPI",
-		"CoordSysAPI",
-		"Cube",
-		"Cylinder",
-		"CylinderLight",
-		"Cylinder_1",
-		"DiskLight",
-		"DistantLight",
-		"DomeLight",
-		"DomeLight_1",
-		"Field3DAsset",
-		"GenerativeProcedural",
-		"GeomModelAPI",
-		"GeomSubset",
-		"GeometryLight",
-		"GroomAPI",
-		"GroomBindingAPI",
-		"HermiteCurves",
-		"HydraGenerativeProceduralAPI",
-		"LightAPI",
-		"LightFilter",
-		"LightListAPI",
-		"ListAPI",
-		"LiveLinkAPI",
-		"Material",
-		"MaterialBindingAPI",
-		"Mesh",
-		"MeshLightAPI",
-		"ModelAPI",
-		"MotionAPI",
-		"NodeDefAPI",
-		"NodeGraph",
-		"NodeGraphNodeAPI",
-		"NurbsCurves",
-		"NurbsPatch",
-		"OpenVDBAsset",
-		"PhysicsArticulationRootAPI",
-		"PhysicsCollisionAPI",
-		"PhysicsCollisionGroup",
-		"PhysicsDistanceJoint",
-		"PhysicsDriveAPI",
-		"PhysicsFilteredPairsAPI",
-		"PhysicsFixedJoint",
-		"PhysicsJoint",
-		"PhysicsLimitAPI",
-		"PhysicsMassAPI",
-		"PhysicsMaterialAPI",
-		"PhysicsMeshCollisionAPI",
-		"PhysicsPrismaticJoint",
-		"PhysicsRevoluteJoint",
-		"PhysicsRigidBodyAPI",
-		"PhysicsScene",
-		"PhysicsSphericalJoint",
-		"Plane",
-		"PluginLight",
-		"PluginLightFilter",
-		"PointInstancer",
-		"Points",
-		"PortalLight",
-		"PrimvarsAPI",
-		"RectLight",
-		"RenderDenoisePass",
-		"RenderPass",
-		"RenderProduct",
-		"RenderSettings",
-		"RenderVar",
-		"RiMaterialAPI",
-		"RiRenderPassAPI",
-		"RiSplineAPI",
-		"SceneGraphPrimAPI",
-		"Scope",
-		"Shader",
-		"ShadowAPI",
-		"ShapingAPI",
-		"SkelAnimation",
-		"SkelBindingAPI",
-		"SkelRoot",
-		"Skeleton",
-		"SparseVolumeTextureAPI",
-		"SpatialAudio",
-		"Sphere",
-		"SphereLight",
-		"StatementsAPI",
-		"TetMesh",
-		"VisibilityAPI",
-		"Volume",
-		"VolumeLightAPI",
-		"Xform",
-		"XformCommonAPI"
-	};
-
-	// clang-format off
-	// "Interesting" here means either something we don't support, or something that would be interesting
-	// to check the usage of, like for our Groom/ControlRig/LiveLink support, or whether the new support
-	// for OpenVDB is used, etc.
-	TMap<FString, int32> InterestingNativeSchemaCounts{
-		{"AssetPreviewsAPI", 0},
-		{"Backdrop", 0},
-		{"BasisCurves", 0},
-		{"Capsule_1", 0},
-		{"ClipsAPI", 0},
-		{"CollectionAPI", 0},
-		{"ControlRigAPI", 0},
-		{"CoordSysAPI", 0},
-		{"CylinderLight", 0},
-		{"Cylinder_1", 0},
-		{"DomeLight_1", 0},
-		{"Field3DAsset", 0},
-		{"GenerativeProcedural", 0},
-		{"GeometryLight", 0},
-		{"GroomAPI", 0},
-		{"GroomBindingAPI", 0},
-		{"HermiteCurves", 0},
-		{"HydraGenerativeProceduralAPI", 0},
-		{"LightFilter", 0},
-		{"LightListAPI", 0},
-		{"ListAPI", 0},
-		{"LiveLinkAPI", 0},
-		{"MeshLightAPI", 0},
-		{"MotionAPI", 0},
-		{"NurbsCurves", 0},
-		{"NurbsPatch", 0},
-		{"OpenVDBAsset", 0},
-		{"PhysicsArticulationRootAPI", 0},
-		{"PhysicsCollisionAPI", 0},
-		{"PhysicsCollisionGroup", 0},
-		{"PhysicsDistanceJoint", 0},
-		{"PhysicsDriveAPI", 0},
-		{"PhysicsFilteredPairsAPI", 0},
-		{"PhysicsFixedJoint", 0},
-		{"PhysicsJoint", 0},
-		{"PhysicsLimitAPI", 0},
-		{"PhysicsMassAPI", 0},
-		{"PhysicsMaterialAPI", 0},
-		{"PhysicsMeshCollisionAPI", 0},
-		{"PhysicsPrismaticJoint", 0},
-		{"PhysicsRevoluteJoint", 0},
-		{"PhysicsRigidBodyAPI", 0},
-		{"PhysicsScene", 0},
-		{"PhysicsSphericalJoint", 0},
-		{"PluginLight", 0},
-		{"PluginLightFilter", 0},
-		{"Points", 0},
-		{"PortalLight", 0},
-		{"RenderDenoisePass", 0},
-		{"RenderPass", 0},
-		{"RenderProduct", 0},
-		{"RenderSettings", 0},
-		{"RenderVar", 0},
-		{"RiMaterialAPI", 0},
-		{"RiRenderPassAPI", 0},
-		{"RiSplineAPI", 0},
-		{"SceneGraphPrimAPI", 0},
-		{"ShadowAPI", 0},
-		{"SparseVolumeTextureAPI", 0},
-		{"SpatialAudio", 0},
-		{"SphereLight", 0},
-		{"StatementsAPI", 0},
-		{"TetMesh", 0},
-		{"VisibilityAPI", 0},
-		{"Volume", 0},
-		{"VolumeLightAPI", 0}
-	};
-	// clang-format on
-
-	TSet<FString> SeenSchemas;
-
-	int32 CustomSchemaCount = 0;
-	{
-		FScopedUsdAllocs Allocs;
-
-		pxr::UsdPrimRange PrimRange{Stage.GetPseudoRoot(), pxr::UsdTraverseInstanceProxies()};
-		for (pxr::UsdPrimRange::iterator PrimRangeIt = ++PrimRange.begin(); PrimRangeIt != PrimRange.end(); ++PrimRangeIt)
-		{
-			// It's perfectly fine to have a typeless prim (e.g. "def 'Cube'")
-			if (PrimRangeIt->HasAuthoredTypeName())
-			{
-				const pxr::TfToken& TypeName = PrimRangeIt->GetTypeName();
-				const FString TypeNameStr = UsdToUnreal::ConvertToken(TypeName);
-
-				if (!TypeNameStr.IsEmpty())
-				{
-					SeenSchemas.Add(TypeNameStr);
-
-					if (int32* InterestingSchemaCount = InterestingNativeSchemaCounts.Find(TypeNameStr))
-					{
-						*InterestingSchemaCount = *InterestingSchemaCount + 1;
-					}
-				}
-			}
-
-			const pxr::UsdPrimTypeInfo& PrimTypeInfo = PrimRangeIt->GetPrimTypeInfo();
-			for (const pxr::TfToken& AppliedSchema : PrimTypeInfo.GetAppliedAPISchemas())
-			{
-				std::pair<pxr::TfToken, pxr::TfToken> Pair = pxr::UsdSchemaRegistry::GetTypeNameAndInstance(AppliedSchema);
-				FString TypeName = UsdToUnreal::ConvertToken(Pair.first);
-
-				// These applied schema names shouldn't ever end up as the empty string... but we don't really want to pop
-				// an ensure or show a warning when analytics fails
-				if (!TypeName.IsEmpty())
-				{
-					SeenSchemas.Add(TypeName);
-
-					if (int32* InterestingSchemaCount = InterestingNativeSchemaCounts.Find(TypeName))
-					{
-						*InterestingSchemaCount = *InterestingSchemaCount + 1;
-					}
-				}
-			}
-		}
-	}
-
-	for (const FString& SchemaName : SeenSchemas)
-	{
-		// We only care about non-native schemas
-		if (!NativeSchemaNames.Contains(SchemaName))
-		{
-			++CustomSchemaCount;
-		}
-	}
-
-	TArray<FAnalyticsEventAttribute> EventAttributes;
-	if (CustomSchemaCount > 0)
-	{
-		EventAttributes.Emplace(TEXT("CustomSchemas"), CustomSchemaCount);
-	}
-
-	for (const TPair<FString, int32>& Pair : InterestingNativeSchemaCounts)
-	{
-		const FString& InterestingNativeSchemaName = Pair.Key;
-		int32 Count = Pair.Value;
-		if (Count > 0)
-		{
-			EventAttributes.Emplace(InterestingNativeSchemaName, Count);
-		}
-	}
+	TArray<FAnalyticsEventAttribute> Attributes;
 
 	FUsdSchemaTranslatorRegistry& Registry = FUsdSchemaTranslatorRegistry::Get();
 	int32 SchemaTranslatorCount = Registry.GetExternalSchemaTranslatorCount();
 	if (SchemaTranslatorCount > 0)
 	{
-		EventAttributes.Emplace(TEXT("CustomSchemaTranslatorCount"), SchemaTranslatorCount);
+		Attributes.Emplace(TEXT("CustomSchemaTranslatorCount"), SchemaTranslatorCount);
 	}
 
-	if (EventAttributes.Num() > 0)
+	UsdUtils::CollectSchemaAnalytics(Stage, Attributes);
+
+	if (Attributes.Num() > 0)
 	{
-		IUsdClassesModule::SendAnalytics(MoveTemp(EventAttributes), FString::Printf(TEXT("%s.CustomSchemaCount"), *EventName));
+		IUsdClassesModule::SendAnalytics(MoveTemp(Attributes), FString::Printf(TEXT("%s.CustomSchemaCount"), *EventName));
 	}
 #endif	  // USE_USD_SDK
 }
diff --git a/Engine/Plugins/Interchange/Extensions/OpenUSD/Source/Import/Private/InterchangeOpenUSDImportModule.cpp b/Engine/Plugins/Interchange/Extensions/OpenUSD/Source/Import/Private/InterchangeOpenUSDImportModule.cpp
--- a/Engine/Plugins/Interchange/Extensions/OpenUSD/Source/Import/Private/InterchangeOpenUSDImportModule.cpp
+++ b/Engine/Plugins/Interchange/Extensions/OpenUSD/Source/Import/Private/InterchangeOpenUSDImportModule.cpp
@@ -217,4 +217,13 @@ void FInterchangeOpenUSDImportModule::ShutdownModule()
 #endif	  // USE_USD_SDK
 }
 
+int32 FInterchangeOpenUSDImportModule::GetNumDefaultHandlers() const
+{
+#if USE_USD_SDK
+	return RegisteredHandlers.Num();
+#else
+	return 0;
+#endif	  // USE_USD_SDK
+}
+
 IMPLEMENT_MODULE(FInterchangeOpenUSDImportModule, InterchangeOpenUSDImport)
\ No newline at end of file
diff --git a/Engine/Plugins/Interchange/Extensions/OpenUSD/Source/Import/Private/InterchangeUsdTranslator.cpp b/Engine/Plugins/Interchange/Extensions/OpenUSD/Source/Import/Private/InterchangeUsdTranslator.cpp
--- a/Engine/Plugins/Interchange/Extensions/OpenUSD/Source/Import/Private/InterchangeUsdTranslator.cpp
+++ b/Engine/Plugins/Interchange/Extensions/OpenUSD/Source/Import/Private/InterchangeUsdTranslator.cpp
@@ -2,6 +2,7 @@
 
 #include "InterchangeUsdTranslator.h"
 
+#include "InterchangeOpenUSDImportModule.h"
 #include "SchemaHandlers/HandlerAccumulatedInfo.h"
 #include "SchemaHandlers/SchemaHandler.h"
 #include "SchemaHandlers/SchemaHandlerRegistry.h"
@@ -29,6 +30,7 @@
 #include "UsdWrappers/UsdTyped.h"
 #include "UsdWrappers/UsdVariantSets.h"
 
+#include "InterchangeAnalyticsHelper.h"
 #include "InterchangeManager.h"
 #include "InterchangeMaterialInstanceNode.h"
 #include "InterchangeMeshNode.h"
@@ -44,6 +46,7 @@
 
 #include "MaterialX/MaterialXUtils/MaterialXManager.h"
 
+#include "AnalyticsEventAttribute.h"
 #include "HAL/IConsoleManager.h"
 #include "Modules/ModuleManager.h"
 #include "MovieSceneSection.h"
@@ -121,6 +124,65 @@ namespace UE::InterchangeUsdTranslator::Private
 		{
 			SourceNode->SetExtraInformation(MetaDataEntry.Key, MetaDataEntry.Value);
 		}
+#endif	  // USE_USD_SDK
+	}
+
+	void AddAnalytics(UE::Interchange::FAnalyticsHelper* AnalyticsHelper, UInterchangeUsdContext* Context)
+	{
+#if USE_USD_SDK
+		if (!AnalyticsHelper || !Context)
+		{
+			return;
+		}
+
+		FInterchangeOpenUSDImportModule& OpenUSDModule = FModuleManager::LoadModuleChecked<FInterchangeOpenUSDImportModule>("InterchangeOpenUSDImport"
+		);
+
+		TArray<FAnalyticsEventAttribute> Attributes;
+		Attributes.Emplace(TEXT("NumTotalHandlers"), UE::Interchange::USD::FSchemaHandlerRegistry::RegisteredHandlerEntries.Num());
+		Attributes.Emplace(
+			TEXT("NumCustomHandlers"),
+			UE::Interchange::USD::FSchemaHandlerRegistry::RegisteredHandlerEntries.Num() - OpenUSDModule.GetNumDefaultHandlers()
+		);
+
+		// TODO: Uncomment when the render context changes are submitted
+		// if (UInterchangeUsdTranslatorSettings* UsdSettings = Cast<UInterchangeUsdTranslatorSettings>(TranslatorSettings))
+		// {
+		// 	Attributes.Emplace(TEXT("NumReorderedHandlers"), UsdSettings->CustomHandlerEntries.Num());
+
+		// 	int32 NumCustomRenderContextHandlers = 0;
+		// 	int32 NumDisabledHandlers = 0;
+		// 	for (const FSchemaHandlerEntry& Entry : UsdSettings->CustomHandlerEntries)
+		// 	{
+		// 		if (!Entry.bEnabled)
+		// 		{
+		// 			NumDisabledHandlers += 1;
+		// 		}
+
+		// 		if (Entry.DefaultRenderContexts != Entry.CustomRenderContexts)
+		// 		{
+		// 			NumCustomRenderContextHandlers += 1;
+		// 		}
+		// 	}
+		// 	Attributes.Emplace(TEXT("NumDisabledHandlers"), NumDisabledHandlers);
+		// 	Attributes.Emplace(TEXT("NumCustomRenderContextHandlers"), NumCustomRenderContextHandlers);
+		// }
+
+		// Clear before adding the events because we may be doing multiple translations on the same import while tweaking the
+		// translator settings or etc., and we don't want to pool duplicate attributes onto the same event
+		const static FString BaseEventIdentifier = TEXT("Interchange.Usage.Import.USD");
+		AnalyticsHelper->ClearAnalyticsEventData(BaseEventIdentifier);
+		AnalyticsHelper->AppendThreadSafe(BaseEventIdentifier, Attributes);
+
+		if (UE::FUsdStage Stage = Context->GetUsdStage())
+		{
+			TArray<FAnalyticsEventAttribute> CustomSchemaAttributes;
+			UsdUtils::CollectSchemaAnalytics(Stage, CustomSchemaAttributes);
+
+			const static FString CustomSchemaEventIdentifier = TEXT("Interchange.Usage.USD.CustomSchemaCount");
+			AnalyticsHelper->ClearAnalyticsEventData(CustomSchemaEventIdentifier);
+			AnalyticsHelper->AppendThreadSafe(CustomSchemaEventIdentifier, CustomSchemaAttributes);
+		}
 #endif	  // USE_USD_SDK
 	}
 }	 // namespace UE::InterchangeUsdTranslator::Private
@@ -431,6 +493,8 @@ bool UInterchangeUSDTranslator::Translate(UInterchangeBaseNodeContainer& NodeCon
 		TranslatePrimSubtree(StageToImport.GetPseudoRoot(), Info);
 	}
 
+	AddAnalytics(AnalyticsHelper.Get(), Context);
+
 	return true;
 #else
 	return false;
diff --git a/Engine/Plugins/Interchange/Extensions/OpenUSD/Source/Import/Public/InterchangeOpenUSDImportModule.h b/Engine/Plugins/Interchange/Extensions/OpenUSD/Source/Import/Public/InterchangeOpenUSDImportModule.h
--- a/Engine/Plugins/Interchange/Extensions/OpenUSD/Source/Import/Public/InterchangeOpenUSDImportModule.h
+++ b/Engine/Plugins/Interchange/Extensions/OpenUSD/Source/Import/Public/InterchangeOpenUSDImportModule.h
@@ -16,6 +16,8 @@ public:
 	virtual void StartupModule() override;
 	virtual void ShutdownModule() override;
 
+	int32 GetNumDefaultHandlers() const;
+
 #if USE_USD_SDK
 private:
 	FDelegateHandle PostEngineInitHandle;
diff --git a/Engine/Plugins/Runtime/USDCore/Source/USDUtilities/Private/USDConversionUtils.cpp b/Engine/Plugins/Runtime/USDCore/Source/USDUtilities/Private/USDConversionUtils.cpp
--- a/Engine/Plugins/Runtime/USDCore/Source/USDUtilities/Private/USDConversionUtils.cpp
+++ b/Engine/Plugins/Runtime/USDCore/Source/USDUtilities/Private/USDConversionUtils.cpp
@@ -25,6 +25,7 @@
 #include "UsdWrappers/UsdStage.h"
 #include "UsdWrappers/VtValue.h"
 
+#include "AnalyticsEventAttribute.h"
 #include "Animation/AnimBlueprint.h"
 #include "Animation/AnimSequence.h"
 #include "Animation/Skeleton.h"
@@ -4311,9 +4312,290 @@ bool UsdUtils::ClearNonEssentialPrimMetadata(const pxr::UsdPrim& Prim)
 }
 #endif	  // USE_USD_SDK
 
-// Deprecated
-void UsdUtils::CollectSchemaAnalytics(const UE::FUsdStage& Stage, const FString& EventName)
+void UsdUtils::CollectSchemaAnalytics(const UE::FUsdStage& Stage, TArray<FAnalyticsEventAttribute>& OutAttributes)
 {
+#if USE_USD_SDK
+	if (!Stage)
+	{
+		return;
+	}
+
+	// Last collected at cl 48758811
+	const static TSet<FString> NativeSchemaNames{
+		"AccessibilityAPI",
+		"AssetPreviewsAPI",
+		"Backdrop",
+		"BasisCurves",
+		"BendingWingAPI",
+		"BlendShape",
+		"Camera",
+		"Capsule",
+		"Capsule_1",
+		"ChaosFabricAPI",
+		"ChaosSolverPropertiesAPI",
+		"ClipsAPI",
+		"CloFabricAPI",
+		"CloSolverPropertiesAPI",
+		"ClothRootAPI",
+		"ClothSolverAPI",
+		"CollapsingAPI",
+		"CollectionAPI",
+		"ColorSpaceAPI",
+		"ColorSpaceDefinitionAPI",
+		"Cone",
+		"ConnectableAPI",
+		"ControlRigAPI",
+		"CoordSysAPI",
+		"Cube",
+		"Cylinder",
+		"CylinderLight",
+		"Cylinder_1",
+		"DiskLight",
+		"DistantLight",
+		"DomeLight",
+		"DomeLight_1",
+		"Field3DAsset",
+		"GenerativeProcedural",
+		"GeomModelAPI",
+		"GeomSubset",
+		"GeometryLight",
+		"GroomAPI",
+		"GroomBindingAPI",
+		"HermiteCurves",
+		"HydraGenerativeProceduralAPI",
+		"LightAPI",
+		"LightFilter",
+		"LightListAPI",
+		"ListAPI",
+		"LiveLinkAPI",
+		"LodSubtreeAPI",
+		"Material",
+		"MaterialBindingAPI",
+		"MaterialXConfigAPI",
+		"Mesh",
+		"MeshLightAPI",
+		"ModelAPI",
+		"MotionAPI",
+		"NaniteAssemblyExternalRefAPI",
+		"NaniteAssemblyRootAPI",
+		"NaniteAssemblySkelBindingAPI",
+		"NodeDefAPI",
+		"NodeGraph",
+		"NodeGraphNodeAPI",
+		"NurbsCurves",
+		"NurbsPatch",
+		"OpenVDBAsset",
+		"PhysicsArticulationRootAPI",
+		"PhysicsCollisionAPI",
+		"PhysicsCollisionGroup",
+		"PhysicsDistanceJoint",
+		"PhysicsDriveAPI",
+		"PhysicsFilteredPairsAPI",
+		"PhysicsFixedJoint",
+		"PhysicsJoint",
+		"PhysicsLimitAPI",
+		"PhysicsMassAPI",
+		"PhysicsMaterialAPI",
+		"PhysicsMeshCollisionAPI",
+		"PhysicsPrismaticJoint",
+		"PhysicsRevoluteJoint",
+		"PhysicsRigidBodyAPI",
+		"PhysicsScene",
+		"PhysicsSphericalJoint",
+		"PinAPI",
+		"Plane",
+		"PluginLight",
+		"PluginLightFilter",
+		"PointInstancer",
+		"Points",
+		"PortalLight",
+		"PrimvarsAPI",
+		"RectLight",
+		"RenderPass",
+		"RenderPatternAPI",
+		"RenderProduct",
+		"RenderSettings",
+		"RenderVar",
+		"RiMaterialAPI",
+		"RiSplineAPI",
+		"SceneGraphPrimAPI",
+		"Scope",
+		"SemanticsLabelsAPI",
+		"SewingAPI",
+		"Shader",
+		"ShadowAPI",
+		"ShapingAPI",
+		"SimMeshDataAPI",
+		"SimPatternAPI",
+		"SkelAnimation",
+		"SkelBindingAPI",
+		"SkelRoot",
+		"Skeleton",
+		"SparseVolumeTextureAPI",
+		"SpatialAudio",
+		"Sphere",
+		"SphereLight",
+		"SpringAPI",
+		"StatementsAPI",
+		"SubdivisionAPI",
+		"TackSpringAPI",
+		"TetMesh",
+		"VisibilityAPI",
+		"Volume",
+		"VolumeLightAPI",
+		"Xform",
+		"XformCommonAPI"
+	};
+
+	// clang-format off
+	// "Interesting" here means either something we don't support, or something that would be interesting
+	// to check the usage of, like for our Groom/ControlRig/LiveLink support, or whether the new support
+	// for OpenVDB is used, etc.
+	TMap<FString, int32> InterestingNativeSchemaCounts{
+		{"AssetPreviewsAPI", 0},
+		{"Backdrop", 0},
+		{"BasisCurves", 0},
+		{"Capsule_1", 0},
+		{"ClipsAPI", 0},
+		{"CollapsingAPI", 0},
+		{"CollectionAPI", 0},
+		{"ControlRigAPI", 0},
+		{"CoordSysAPI", 0},
+		{"Cylinder_1", 0},
+		{"CylinderLight", 0},
+		{"Field3DAsset", 0},
+		{"GenerativeProcedural", 0},
+		{"GeometryLight", 0},
+		{"GroomAPI", 0},
+		{"GroomBindingAPI", 0},
+		{"HermiteCurves", 0},
+		{"HydraGenerativeProceduralAPI", 0},
+		{"LightFilter", 0},
+		{"LightListAPI", 0},
+		{"ListAPI", 0},
+		{"LiveLinkAPI", 0},
+		{"LodSubtreeAPI", 0},
+		{"MeshLightAPI", 0},
+		{"MotionAPI", 0},
+		{"NaniteAssemblyExternalRefAPI", 0},
+		{"NaniteAssemblyRootAPI", 0},
+		{"NaniteAssemblySkelBindingAPI", 0},
+		{"NurbsCurves", 0},
+		{"NurbsPatch", 0},
+		{"OpenVDBAsset", 0},
+		{"PhysicsArticulationRootAPI", 0},
+		{"PhysicsCollisionAPI", 0},
+		{"PhysicsCollisionGroup", 0},
+		{"PhysicsDistanceJoint", 0},
+		{"PhysicsDriveAPI", 0},
+		{"PhysicsFilteredPairsAPI", 0},
+		{"PhysicsFixedJoint", 0},
+		{"PhysicsJoint", 0},
+		{"PhysicsLimitAPI", 0},
+		{"PhysicsMassAPI", 0},
+		{"PhysicsMaterialAPI", 0},
+		{"PhysicsMeshCollisionAPI", 0},
+		{"PhysicsPrismaticJoint", 0},
+		{"PhysicsRevoluteJoint", 0},
+		{"PhysicsRigidBodyAPI", 0},
+		{"PhysicsScene", 0},
+		{"PhysicsSphericalJoint", 0},
+		{"PluginLight", 0},
+		{"PluginLightFilter", 0},
+		{"Points", 0},
+		{"PortalLight", 0},
+		{"RenderDenoisePass", 0},
+		{"RenderPass", 0},
+		{"RenderProduct", 0},
+		{"RenderSettings", 0},
+		{"RenderVar", 0},
+		{"RiMaterialAPI", 0},
+		{"RiRenderPassAPI", 0},
+		{"RiSplineAPI", 0},
+		{"SceneGraphPrimAPI", 0},
+		{"ShadowAPI", 0},
+		{"SparseVolumeTextureAPI", 0},
+		{"SpatialAudio", 0},
+		{"StatementsAPI", 0},
+		{"SubdivisionAPI", 0},
+		{"TetMesh", 0},
+		{"VisibilityAPI", 0},
+		{"Volume", 0},
+		{"VolumeLightAPI", 0}
+	};
+	// clang-format on
+
+	TSet<FString> SeenSchemas;
+
+	int32 CustomSchemaCount = 0;
+	{
+		FScopedUsdAllocs Allocs;
+
+		pxr::UsdPrimRange PrimRange{Stage.GetPseudoRoot(), pxr::UsdTraverseInstanceProxies()};
+		for (pxr::UsdPrimRange::iterator PrimRangeIt = ++PrimRange.begin(); PrimRangeIt != PrimRange.end(); ++PrimRangeIt)
+		{
+			// It's perfectly fine to have a typeless prim (e.g. "def 'Cube'")
+			if (PrimRangeIt->HasAuthoredTypeName())
+			{
+				const pxr::TfToken& TypeName = PrimRangeIt->GetTypeName();
+				const FString TypeNameStr = UsdToUnreal::ConvertToken(TypeName);
+
+				if (!TypeNameStr.IsEmpty())
+				{
+					SeenSchemas.Add(TypeNameStr);
+
+					if (int32* InterestingSchemaCount = InterestingNativeSchemaCounts.Find(TypeNameStr))
+					{
+						*InterestingSchemaCount = *InterestingSchemaCount + 1;
+					}
+				}
+			}
+
+			const pxr::UsdPrimTypeInfo& PrimTypeInfo = PrimRangeIt->GetPrimTypeInfo();
+			for (const pxr::TfToken& AppliedSchema : PrimTypeInfo.GetAppliedAPISchemas())
+			{
+				std::pair<pxr::TfToken, pxr::TfToken> Pair = pxr::UsdSchemaRegistry::GetTypeNameAndInstance(AppliedSchema);
+				FString TypeName = UsdToUnreal::ConvertToken(Pair.first);
+
+				// These applied schema names shouldn't ever end up as the empty string... but we don't really want to pop
+				// an ensure or show a warning when analytics fails
+				if (!TypeName.IsEmpty())
+				{
+					SeenSchemas.Add(TypeName);
+
+					if (int32* InterestingSchemaCount = InterestingNativeSchemaCounts.Find(TypeName))
+					{
+						*InterestingSchemaCount = *InterestingSchemaCount + 1;
+					}
+				}
+			}
+		}
+	}
+
+	for (const FString& SchemaName : SeenSchemas)
+	{
+		// We only care about non-native schemas
+		if (!NativeSchemaNames.Contains(SchemaName))
+		{
+			++CustomSchemaCount;
+		}
+	}
+
+	if (CustomSchemaCount > 0)
+	{
+		OutAttributes.Emplace(TEXT("CustomSchemas"), CustomSchemaCount);
+	}
+
+	for (const TPair<FString, int32>& Pair : InterestingNativeSchemaCounts)
+	{
+		const FString& InterestingNativeSchemaName = Pair.Key;
+		int32 Count = Pair.Value;
+		if (Count > 0)
+		{
+			OutAttributes.Emplace(InterestingNativeSchemaName, Count);
+		}
+	}
+#endif	  // USE_USD_SDK
 }
 
 void UsdUtils::ReadStageMetaData(UE::FUsdStage Stage, TMap<FString, FString>& OutMetaDataMap)
diff --git a/Engine/Plugins/Runtime/USDCore/Source/USDUtilities/Public/USDConversionUtils.h b/Engine/Plugins/Runtime/USDCore/Source/USDUtilities/Public/USDConversionUtils.h
--- a/Engine/Plugins/Runtime/USDCore/Source/USDUtilities/Public/USDConversionUtils.h
+++ b/Engine/Plugins/Runtime/USDCore/Source/USDUtilities/Public/USDConversionUtils.h
@@ -58,6 +58,7 @@ class UUsdAssetUserData;
 enum class EUsdDrawMode : int32;
 enum class EUsdDuplicateType : uint8;
 enum class EUsdUpAxis : uint8;
+struct FAnalyticsEventAttribute;
 struct FUsdUnrealAssetInfo;
 template<typename UsdObjectType>
 class TUsdStore;
@@ -707,9 +708,8 @@ namespace UsdUtils
 	USDUTILITIES_API bool ClearNonEssentialPrimMetadata(const pxr::UsdPrim& Prim);
 #endif	  // USE_USD_SDK
 
-	/** Collects how many times each schema shows up on the provided stage and send it as an analytics event */
-	UE_DEPRECATED(5.5, "This function has been moved to the USDSchemasModule")
-	USDUTILITIES_API void CollectSchemaAnalytics(const UE::FUsdStage& Stage, const FString& EventName);
+	/** Collects how many times each schema shows up on the provided stage and appends those attributes to the provided OutAttributes */
+	USDUTILITIES_API void CollectSchemaAnalytics(const UE::FUsdStage& Stage, TArray<FAnalyticsEventAttribute>& OutAttributes);
 
 	/** Parses and stores primary Stage metadata key-info pairs in the MetadataMap */
 	USDUTILITIES_API void ReadStageMetaData(UE::FUsdStage Stage, TMap<FString, FString>& OutMetaDataMap);
@@ -719,7 +719,7 @@ namespace UsdUtils
 
 	// Retrieves the equivalent of this prim in the prototype it's referenced from, where applicable
 	USDUTILITIES_API UE::FUsdPrim GetPrototypePrim(const UE::FUsdPrim& Prim);
-	
+
 	// Retrieves the prototype path for the prim
 	USDUTILITIES_API UE::FSdfPath GetPrototypePrimPath(const UE::FUsdPrim& Prim);
 
