USD: Import DomeLight_1 schemas as DomeLight schemas.

Export DomeLight_1 schemas when writing to USD.

#Synced-CL 49553638

diff --git a/Engine/Plugins/Importers/USDImporter/Source/USDSchemas/Private/USDLuxLightTranslator.cpp b/Engine/Plugins/Importers/USDImporter/Source/USDSchemas/Private/USDLuxLightTranslator.cpp
--- a/Engine/Plugins/Importers/USDImporter/Source/USDSchemas/Private/USDLuxLightTranslator.cpp
+++ b/Engine/Plugins/Importers/USDImporter/Source/USDSchemas/Private/USDLuxLightTranslator.cpp
@@ -26,6 +26,7 @@
 #include "USDIncludesStart.h"
 #include "pxr/usd/usdLux/diskLight.h"
 #include "pxr/usd/usdLux/domeLight.h"
+#include "pxr/usd/usdLux/domeLight_1.h"
 #include "pxr/usd/usdLux/lightAPI.h"
 #include "pxr/usd/usdLux/rectLight.h"
 #include "USDIncludesEnd.h"
@@ -51,28 +52,36 @@ void FUsdLuxLightTranslator::CreateAssets()
 
 	const FString VolumePrimPathString = PrimPath.GetString();
 	pxr::UsdPrim Prim = GetPrim();
-	pxr::UsdLuxDomeLight DomeLight{Prim};
-	if (!DomeLight)
+
+	pxr::UsdAttribute TextureFileAttr;
+	if (pxr::UsdLuxDomeLight DomeLight{Prim})
+	{
+		TextureFileAttr = DomeLight.GetTextureFileAttr();
+	}
+	else if (pxr::UsdLuxDomeLight_1 DomeLight1{Prim})
+	{
+		TextureFileAttr = DomeLight1.GetTextureFileAttr();
+	}
+	if (!TextureFileAttr)
 	{
-		// Only dome lights make assets for now
 		return;
 	}
 
-	const FString ResolvedDomeTexturePath = UsdUtils::GetResolvedAssetPath(DomeLight.GetTextureFileAttr());
+	const FString ResolvedDomeTexturePath = UsdUtils::GetResolvedAssetPath(TextureFileAttr);
 	if (ResolvedDomeTexturePath.IsEmpty())
 	{
 		FScopedUsdAllocs Allocs;
 
 		pxr::SdfAssetPath TextureAssetPath;
-		DomeLight.GetTextureFileAttr().Get<pxr::SdfAssetPath>(&TextureAssetPath);
+		TextureFileAttr.Get<pxr::SdfAssetPath>(&TextureAssetPath);
 
 		// Show a good warning for this because it's easy to pick some cubemap asset from the engine (that usually don't come with the
 		// source texture) and have the dome light silently not work again
 		FString TargetAssetPath = UsdToUnreal::ConvertString(TextureAssetPath.GetAssetPath());
 		USD_LOG_USERWARNING(FText::Format(
-			LOCTEXT("FailToFindTexture", "Failed to find texture '{0}' used for UsdLuxDomeLight '{1}'!"),
+			LOCTEXT("FailToFindTexture", "Failed to find texture '{0}' used for DomeLight '{1}'!"),
 			FText::FromString(TargetAssetPath),
-			FText::FromString(UsdToUnreal::ConvertPath(DomeLight.GetPrim().GetPath()))
+			FText::FromString(UsdToUnreal::ConvertPath(Prim.GetPath()))
 		));
 
 		return;
@@ -99,7 +108,7 @@ void FUsdLuxLightTranslator::CreateAssets()
 
 	if (UUsdAssetUserData* TextureUserData = UsdUnreal::ObjectUtils::GetOrCreateAssetUserData(Texture))
 	{
-		TextureUserData->PrimPaths.AddUnique(UsdToUnreal::ConvertPath(DomeLight.GetPrim().GetPath()));
+		TextureUserData->PrimPaths.AddUnique(UsdToUnreal::ConvertPath(Prim.GetPath()));
 	}
 
 	if (Texture)
diff --git a/Engine/Plugins/Interchange/Extensions/OpenUSD/Source/Import/Private/SchemaHandlers/LightSchemaHandler.cpp b/Engine/Plugins/Interchange/Extensions/OpenUSD/Source/Import/Private/SchemaHandlers/LightSchemaHandler.cpp
--- a/Engine/Plugins/Interchange/Extensions/OpenUSD/Source/Import/Private/SchemaHandlers/LightSchemaHandler.cpp
+++ b/Engine/Plugins/Interchange/Extensions/OpenUSD/Source/Import/Private/SchemaHandlers/LightSchemaHandler.cpp
@@ -84,7 +84,7 @@ namespace UE::LuxLightSchemaHandler::Private
 		{
 			return UInterchangeRectLightNode::StaticClass();
 		}
-		else if (Prim.IsA(TEXT("DomeLight")))
+		else if (Prim.IsA(TEXT("DomeLight")) || Prim.IsA(TEXT("DomeLight_1")))
 		{
 			return UInterchangeSkyLightNode::StaticClass();
 		}
@@ -303,7 +303,7 @@ namespace UE::LuxLightSchemaHandler::Private
 				}
 			}
 		}
-		else if (Prim.IsA(TEXT("DomeLight")))
+		else if (Prim.IsA(TEXT("DomeLight")) || Prim.IsA(TEXT("DomeLight_1")))
 		{
 			UInterchangeSkyLightNode* CastLightNode = Cast<UInterchangeSkyLightNode>(LightNode);
 			if (!ensure(CastLightNode))
diff --git a/Engine/Plugins/Runtime/USDCore/Source/USDUtilities/Private/USDConversionUtils.cpp b/Engine/Plugins/Runtime/USDCore/Source/USDUtilities/Private/USDConversionUtils.cpp
--- a/Engine/Plugins/Runtime/USDCore/Source/USDUtilities/Private/USDConversionUtils.cpp
+++ b/Engine/Plugins/Runtime/USDCore/Source/USDUtilities/Private/USDConversionUtils.cpp
@@ -95,6 +95,7 @@
 #include "pxr/usd/usdLux/diskLight.h"
 #include "pxr/usd/usdLux/distantLight.h"
 #include "pxr/usd/usdLux/domeLight.h"
+#include "pxr/usd/usdLux/domeLight_1.h"
 #include "pxr/usd/usdLux/rectLight.h"
 #include "pxr/usd/usdLux/shapingAPI.h"
 #include "pxr/usd/usdLux/sphereLight.h"
@@ -538,7 +539,7 @@ UClass* UsdUtils::GetActorTypeForPrim(const pxr::UsdPrim& Prim)
 			return APointLight::StaticClass();
 		}
 	}
-	else if (Prim.IsA<pxr::UsdLuxDomeLight>())
+	else if (Prim.IsA<pxr::UsdLuxDomeLight>() || Prim.IsA<pxr::UsdLuxDomeLight_1>())
 	{
 		return ASkyLight::StaticClass();
 	}
@@ -589,7 +590,7 @@ UClass* UsdUtils::GetComponentTypeForPrim(const pxr::UsdPrim& Prim)
 			return UPointLightComponent::StaticClass();
 		}
 	}
-	else if (Prim.IsA<pxr::UsdLuxDomeLight>())
+	else if (Prim.IsA<pxr::UsdLuxDomeLight>() || Prim.IsA<pxr::UsdLuxDomeLight_1>())
 	{
 		return USkyLightComponent::StaticClass();
 	}
@@ -655,7 +656,7 @@ FString UsdUtils::GetSchemaNameForComponent(const USceneComponent& Component)
 	}
 	else if (Component.IsA<USkyLightComponent>())
 	{
-		return TEXT("DomeLight");
+		return TEXT("DomeLight_1");
 	}
 	else if (Component.IsA<UAudioComponent>())
 	{
diff --git a/Engine/Plugins/Runtime/USDCore/Source/USDUtilities/Private/USDLightConversion.cpp b/Engine/Plugins/Runtime/USDCore/Source/USDUtilities/Private/USDLightConversion.cpp
--- a/Engine/Plugins/Runtime/USDCore/Source/USDUtilities/Private/USDLightConversion.cpp
+++ b/Engine/Plugins/Runtime/USDCore/Source/USDUtilities/Private/USDLightConversion.cpp
@@ -34,6 +34,7 @@
 #include "pxr/usd/usdLux/diskLight.h"
 #include "pxr/usd/usdLux/distantLight.h"
 #include "pxr/usd/usdLux/domeLight.h"
+#include "pxr/usd/usdLux/domeLight_1.h"
 #include "pxr/usd/usdLux/lightAPI.h"
 #include "pxr/usd/usdLux/rectLight.h"
 #include "pxr/usd/usdLux/shadowAPI.h"
@@ -225,8 +226,16 @@ bool UsdToUnreal::ConvertDomeLight(
 {
 	FScopedUsdAllocs UsdAllocs;
 
-	pxr::UsdLuxDomeLight DomeLight{Prim};
-	if (!DomeLight)
+	pxr::UsdAttribute TextureFileAttr;
+	if (pxr::UsdLuxDomeLight DomeLight{Prim})
+	{
+		TextureFileAttr = DomeLight.GetTextureFileAttr();
+	}
+	else if (pxr::UsdLuxDomeLight_1 DomeLight1{Prim})
+	{
+		TextureFileAttr = DomeLight1.GetTextureFileAttr();
+	}
+	if (!TextureFileAttr)
 	{
 		return false;
 	}
@@ -234,21 +243,21 @@ bool UsdToUnreal::ConvertDomeLight(
 	// Revert the allocator in case we end up creating a texture on the ansi allocator or something like that
 	FScopedUnrealAllocs UEAllocs;
 
-	const FString ResolvedDomeTexturePath = UsdUtils::GetResolvedAssetPath(DomeLight.GetTextureFileAttr());
+	const FString ResolvedDomeTexturePath = UsdUtils::GetResolvedAssetPath(TextureFileAttr);
 	if (ResolvedDomeTexturePath.IsEmpty())
 	{
 		FScopedUsdAllocs Allocs;
 
 		pxr::SdfAssetPath TextureAssetPath;
-		DomeLight.GetTextureFileAttr().Get<pxr::SdfAssetPath>(&TextureAssetPath);
+		TextureFileAttr.Get<pxr::SdfAssetPath>(&TextureAssetPath);
 
 		// Show a good warning for this because it's easy to pick some cubemap asset from the engine (that usually don't come with the
 		// source texture) and have the dome light silently not work again
 		FString TargetAssetPath = UsdToUnreal::ConvertString(TextureAssetPath.GetAssetPath());
 		USD_LOG_USERWARNING(FText::Format(
-			LOCTEXT("MissingDomeLightTexture", "Failed to find texture '{0}' used for UsdLuxDomeLight '{1}'!"),
+			LOCTEXT("MissingDomeLightTexture", "Failed to find texture '{0}' used for DomeLight '{1}'!"),
 			FText::FromString(TargetAssetPath),
-			FText::FromString(UsdToUnreal::ConvertPath(DomeLight.GetPrim().GetPath()))
+			FText::FromString(UsdToUnreal::ConvertPath(Prim.GetPath()))
 		));
 
 		return true;
@@ -269,7 +278,7 @@ bool UsdToUnreal::ConvertDomeLight(
 
 	if (UUsdAssetUserData* TextureUserData = UsdUnreal::ObjectUtils::GetOrCreateAssetUserData(Cubemap))
 	{
-		TextureUserData->PrimPaths.AddUnique(UsdToUnreal::ConvertPath(DomeLight.GetPrim().GetPath()));
+		TextureUserData->PrimPaths.AddUnique(UsdToUnreal::ConvertPath(Prim.GetPath()));
 	}
 
 	if (Cubemap)
@@ -637,16 +646,20 @@ bool UnrealToUsd::ConvertSkyLightComponent(const USkyLightComponent& LightCompon
 		return false;
 	}
 
-	pxr::UsdLuxDomeLight Light{Prim};
-	if (!Light)
+	pxr::UsdAttribute TextureFileAttr;
+	if (pxr::UsdLuxDomeLight DomeLight{Prim})
 	{
-		return false;
+		TextureFileAttr = DomeLight.GetTextureFileAttr();
+	}
+	else if (pxr::UsdLuxDomeLight_1 DomeLight1{Prim})
+	{
+		TextureFileAttr = DomeLight1.GetTextureFileAttr();
 	}
 
 #if WITH_EDITORONLY_DATA
 	FUsdStageInfo StageInfo(Prim.GetStage());
 
-	if (pxr::UsdAttribute Attr = Light.CreateTextureFileAttr())
+	if (TextureFileAttr)
 	{
 		if (UTextureCube* TextureCube = LightComponent.Cubemap)
 		{
@@ -667,8 +680,8 @@ bool UnrealToUsd::ConvertSkyLightComponent(const USkyLightComponent& LightCompon
 				}
 
 				UsdUtils::MakePathRelativeToLayer(UE::FSdfLayer{Prim.GetStage()->GetEditTarget().GetLayer()}, FilePath);
-				Attr.Set<pxr::SdfAssetPath>(pxr::SdfAssetPath{UnrealToUsd::ConvertString(*FilePath).Get()}, UsdTimeCode);
-				UsdUtils::NotifyIfOverriddenOpinion(Attr);
+				TextureFileAttr.Set<pxr::SdfAssetPath>(pxr::SdfAssetPath{UnrealToUsd::ConvertString(*FilePath).Get()}, UsdTimeCode);
+				UsdUtils::NotifyIfOverriddenOpinion(TextureFileAttr);
 			}
 		}
 	}
diff --git a/Engine/Plugins/Runtime/USDCore/Source/USDUtilities/Private/USDPrimConversion.cpp b/Engine/Plugins/Runtime/USDCore/Source/USDUtilities/Private/USDPrimConversion.cpp
--- a/Engine/Plugins/Runtime/USDCore/Source/USDUtilities/Private/USDPrimConversion.cpp
+++ b/Engine/Plugins/Runtime/USDCore/Source/USDUtilities/Private/USDPrimConversion.cpp
@@ -104,6 +104,7 @@
 #include "pxr/usd/usdLux/diskLight.h"
 #include "pxr/usd/usdLux/distantLight.h"
 #include "pxr/usd/usdLux/domeLight.h"
+#include "pxr/usd/usdLux/domeLight_1.h"
 #include "pxr/usd/usdLux/lightAPI.h"
 #include "pxr/usd/usdLux/rectLight.h"
 #include "pxr/usd/usdLux/shadowAPI.h"
@@ -2017,12 +2018,14 @@ UsdToUnreal::FPropertyTrackReader UsdToUnreal::CreatePropertyTrackReader(
 				}
 			}
 		}
-		else if (pxr::UsdLuxDomeLight DomeLight{UsdPrim})
+		else if (UsdPrim.IsA<pxr::UsdLuxDomeLight>() || UsdPrim.IsA<pxr::UsdLuxDomeLight_1>())
 		{
+			pxr::UsdLuxNonboundableLightBase LightBase{UsdPrim};
+
 			if (PropertyPath == UnrealIdentifiers::IntensityPropertyName)
 			{
-				pxr::UsdAttribute IntensityAttr = SphereLight.GetIntensityAttr();
-				pxr::UsdAttribute ExposureAttr = SphereLight.GetExposureAttr();
+				pxr::UsdAttribute IntensityAttr = LightBase.GetIntensityAttr();
+				pxr::UsdAttribute ExposureAttr = LightBase.GetExposureAttr();
 
 				if (IntensityAttr && ExposureAttr)
 				{
@@ -6714,7 +6717,7 @@ TArray<UE::FUsdAttribute> UsdUtils::GetAttributesForProperty(const UE::FUsdPrim&
 				UE::FUsdAttribute{UsdPrim.GetAttribute(pxr::UsdLuxTokens->inputsRadius)}
 			};
 		}
-		else if (UsdPrim.IsA<pxr::UsdLuxDistantLight>() || UsdPrim.IsA<pxr::UsdLuxDomeLight>())
+		else if (UsdPrim.IsA<pxr::UsdLuxDistantLight>() || UsdPrim.IsA<pxr::UsdLuxDomeLight>() || UsdPrim.IsA<pxr::UsdLuxDomeLight_1>())
 		{
 			return {
 				UE::FUsdAttribute{UsdPrim.GetAttribute(pxr::UsdLuxTokens->inputsIntensity)},
