Interchange: Add support for nits as light units.

USD: Support converting light intensity values to/from nits.

#virtualized
#Synced-CL 49553638

diff --git a/Engine/Plugins/Interchange/Extensions/OpenUSD/Source/Import/Private/SchemaHandlers/LightSchemaHandler.cpp b/Engine/Plugins/Interchange/Extensions/OpenUSD/Source/Import/Private/SchemaHandlers/LightSchemaHandler.cpp
--- a/Engine/Plugins/Interchange/Extensions/OpenUSD/Source/Import/Private/SchemaHandlers/LightSchemaHandler.cpp
+++ b/Engine/Plugins/Interchange/Extensions/OpenUSD/Source/Import/Private/SchemaHandlers/LightSchemaHandler.cpp
@@ -231,28 +231,19 @@ namespace UE::LuxLightSchemaHandler::Private
 					CastLightNode->SetCustomInnerConeAngle(InnerConeAngle);
 					CastLightNode->SetCustomOuterConeAngle(OuterConeAngle);
 				}
-
-				if (bIntensityAuthored)
-				{
-					Intensity = UsdToUnreal::ConvertLuxShapingAPIIntensityAttr(Intensity, Exposure, Radius, ConeAngle, ConeSoftness, StageInfo);
-					CastLightNode->SetCustomIntensityUnits(EInterchangeLightUnits::Lumens);
-					CastLightNode->SetCustomIntensity(Intensity);
-				}
 			}
-			else
+
+			UInterchangePointLightNode* CastLightNode = Cast<UInterchangePointLightNode>(LightNode);
+			if (!ensure(CastLightNode))
 			{
-				UInterchangePointLightNode* CastLightNode = Cast<UInterchangePointLightNode>(LightNode);
-				if (!ensure(CastLightNode))
-				{
-					return;
-				}
+				return;
+			}
 
-				if (bIntensityAuthored)
-				{
-					Intensity = UsdToUnreal::ConvertSphereLightIntensityAttr(Intensity, Exposure, Radius, StageInfo);
-					CastLightNode->SetCustomIntensityUnits(EInterchangeLightUnits::Lumens);
-					CastLightNode->SetCustomIntensity(Intensity);
-				}
+			if (bIntensityAuthored)
+			{
+				Intensity = UsdToUnreal::ConvertLightIntensityAttr(Intensity, Exposure);
+				CastLightNode->SetCustomIntensityUnits(EInterchangeLightUnits::Nits);
+				CastLightNode->SetCustomIntensity(Intensity);
 			}
 		}
 		else if (Prim.IsA(TEXT("RectLight")) || Prim.IsA(TEXT("DiskLight")))
@@ -286,8 +277,8 @@ namespace UE::LuxLightSchemaHandler::Private
 
 				if (bIntensityAuthored)
 				{
-					Intensity = UsdToUnreal::ConvertRectLightIntensityAttr(Intensity, Exposure, Width, Height, StageInfo);
-					CastLightNode->SetCustomIntensityUnits(EInterchangeLightUnits::Lumens);
+					Intensity = UsdToUnreal::ConvertLightIntensityAttr(Intensity, Exposure);
+					CastLightNode->SetCustomIntensityUnits(EInterchangeLightUnits::Nits);
 					CastLightNode->SetCustomIntensity(Intensity);
 				}
 			}
@@ -304,8 +295,8 @@ namespace UE::LuxLightSchemaHandler::Private
 
 				if (bIntensityAuthored)
 				{
-					Intensity = UsdToUnreal::ConvertDiskLightIntensityAttr(Intensity, Exposure, Radius, StageInfo);
-					CastLightNode->SetCustomIntensityUnits(EInterchangeLightUnits::Lumens);
+					Intensity = UsdToUnreal::ConvertLightIntensityAttr(Intensity, Exposure);
+					CastLightNode->SetCustomIntensityUnits(EInterchangeLightUnits::Nits);
 					CastLightNode->SetCustomIntensity(Intensity);
 				}
 			}
diff --git a/Engine/Plugins/Interchange/Runtime/Source/Nodes/Public/InterchangeLightNode.h b/Engine/Plugins/Interchange/Runtime/Source/Nodes/Public/InterchangeLightNode.h
--- a/Engine/Plugins/Interchange/Runtime/Source/Nodes/Public/InterchangeLightNode.h
+++ b/Engine/Plugins/Interchange/Runtime/Source/Nodes/Public/InterchangeLightNode.h
@@ -17,6 +17,7 @@ enum class EInterchangeLightUnits : uint8
 	Candelas,
 	Lumens,
 	EV,
+	Nits,
 };
 
 // Mirrors ESkyLightSourceType
diff --git a/Engine/Plugins/Interchange/Runtime/Source/Pipelines/Private/InterchangeGenericScenesPipeline.cpp b/Engine/Plugins/Interchange/Runtime/Source/Pipelines/Private/InterchangeGenericScenesPipeline.cpp
--- a/Engine/Plugins/Interchange/Runtime/Source/Pipelines/Private/InterchangeGenericScenesPipeline.cpp
+++ b/Engine/Plugins/Interchange/Runtime/Source/Pipelines/Private/InterchangeGenericScenesPipeline.cpp
@@ -1059,6 +1059,7 @@ void UInterchangeGenericLevelPipeline::SetUpFactoryNode(UInterchangeActorFactory
 			static_assert(FCommonLightUnits(EInterchangeLightUnits::Lumens) == FCommonLightUnits(ELightUnits::Lumens), "EInterchangeLightUnits::Lumens differs from ELightUnits::Lumens");
 			static_assert(FCommonLightUnits(EInterchangeLightUnits::Candelas) == FCommonLightUnits(ELightUnits::Candelas), "EInterchangeLightUnits::Candelas differs from ELightUnits::Candelas");
 			static_assert(FCommonLightUnits(EInterchangeLightUnits::EV) == FCommonLightUnits(ELightUnits::EV), "EInterchangeLightUnits::EV differs from ELightUnits::EV");
+			static_assert(FCommonLightUnits(EInterchangeLightUnits::Nits) == FCommonLightUnits(ELightUnits::Nits), "EInterchangeLightUnits::Nits differs from ELightUnits::Nits");
 
 			if (const UInterchangeLightNode* LightNode = Cast<UInterchangeLightNode>(BaseLightNode))
 			{
diff --git a/Engine/Plugins/Runtime/USDCore/Source/USDUtilities/Private/USDLightConversion.cpp b/Engine/Plugins/Runtime/USDCore/Source/USDUtilities/Private/USDLightConversion.cpp
--- a/Engine/Plugins/Runtime/USDCore/Source/USDUtilities/Private/USDLightConversion.cpp
+++ b/Engine/Plugins/Runtime/USDCore/Source/USDUtilities/Private/USDLightConversion.cpp
@@ -143,8 +143,8 @@ bool UsdToUnreal::ConvertRectLight(const pxr::UsdPrim& Prim, URectLightComponent
 
 	LightComponent.SourceWidth = UsdToUnreal::ConvertDistance(StageInfo, UsdWidth);
 	LightComponent.SourceHeight = UsdToUnreal::ConvertDistance(StageInfo, UsdHeight);
-	LightComponent.Intensity = UsdToUnreal::ConvertRectLightIntensityAttr(UsdIntensity, UsdExposure, UsdWidth, UsdHeight, StageInfo);
-	LightComponent.IntensityUnits = ELightUnits::Lumens;
+	LightComponent.Intensity = UsdToUnreal::ConvertLightIntensityAttr(UsdIntensity, UsdExposure);
+	LightComponent.IntensityUnits = ELightUnits::Nits;
 
 	return true;
 }
@@ -167,8 +167,8 @@ bool UsdToUnreal::ConvertDiskLight(const pxr::UsdPrim& Prim, URectLightComponent
 
 	LightComponent.SourceWidth = UsdToUnreal::ConvertDistance(StageInfo, UsdRadius) * 2.f;
 	LightComponent.SourceHeight = LightComponent.SourceWidth;
-	LightComponent.Intensity = UsdToUnreal::ConvertDiskLightIntensityAttr(UsdIntensity, UsdExposure, UsdRadius, StageInfo);
-	LightComponent.IntensityUnits = ELightUnits::Lumens;
+	LightComponent.Intensity = UsdToUnreal::ConvertLightIntensityAttr(UsdIntensity, UsdExposure);
+	LightComponent.IntensityUnits = ELightUnits::Nits;
 
 	return true;
 }
@@ -189,29 +189,8 @@ bool UsdToUnreal::ConvertSphereLight(const pxr::UsdPrim& Prim, UPointLightCompon
 	const float UsdExposure = UsdUtils::GetUsdValue<float>(SphereLight.GetExposureAttr(), UsdTimeCode);
 	const float UsdRadius = UsdUtils::GetUsdValue<float>(SphereLight.GetRadiusAttr(), UsdTimeCode);
 
-	if (pxr::UsdLuxShapingAPI ShapingAPI{Prim})
-	{
-		const float UsdConeAngle = UsdUtils::GetUsdValue<float>(ShapingAPI.GetShapingConeAngleAttr(), UsdTimeCode);
-		const float UsdConeSoftness = UsdUtils::GetUsdValue<float>(ShapingAPI.GetShapingConeSoftnessAttr(), UsdTimeCode);
-
-		float InnerConeAngle = 0.0f;
-		const float OuterConeAngle = UsdToUnreal::ConvertConeAngleSoftnessAttr(UsdConeAngle, UsdConeSoftness, InnerConeAngle);
-
-		LightComponent.Intensity = UsdToUnreal::ConvertLuxShapingAPIIntensityAttr(
-			UsdIntensity,
-			UsdExposure,
-			UsdRadius,
-			UsdConeAngle,
-			UsdConeSoftness,
-			StageInfo
-		);
-	}
-	else
-	{
-		LightComponent.Intensity = UsdToUnreal::ConvertSphereLightIntensityAttr(UsdIntensity, UsdExposure, UsdRadius, StageInfo);
-	}
-
-	LightComponent.IntensityUnits = ELightUnits::Lumens;
+	LightComponent.Intensity = UsdToUnreal::ConvertLightIntensityAttr(UsdIntensity, UsdExposure);
+	LightComponent.IntensityUnits = ELightUnits::Nits;
 	LightComponent.SourceRadius = UsdToUnreal::ConvertDistance(StageInfo, UsdRadius);
 
 	return true;
@@ -839,6 +818,8 @@ float UsdUtils::ConvertIntensityToNits(float Intensity, float Steradians, float
 		case ELightUnits::EV:
 			// Nit = luminance (cd/m2)
 			return EV100ToLuminance(Intensity);
+		case ELightUnits::Nits:
+			return Intensity;
 		case ELightUnits::Unitless:
 			// Nit = ( unitless / 625 ) / area = candela / area
 			// https://docs.unrealengine.com/en-US/Engine/Rendering/LightingAndShadows/PhysicalLightUnits/index.html#point,spot,andrectlights
diff --git a/Engine/Plugins/Runtime/USDCore/Source/USDUtilities/Private/USDPrimConversion.cpp b/Engine/Plugins/Runtime/USDCore/Source/USDUtilities/Private/USDPrimConversion.cpp
--- a/Engine/Plugins/Runtime/USDCore/Source/USDUtilities/Private/USDPrimConversion.cpp
+++ b/Engine/Plugins/Runtime/USDCore/Source/USDUtilities/Private/USDPrimConversion.cpp
@@ -1658,56 +1658,27 @@ UsdToUnreal::FPropertyTrackReader UsdToUnreal::CreatePropertyTrackReader(
 				{
 					pxr::UsdAttribute IntensityAttr = SphereLight.GetIntensityAttr();
 					pxr::UsdAttribute ExposureAttr = SphereLight.GetExposureAttr();
-					pxr::UsdAttribute RadiusAttr = SphereLight.GetRadiusAttr();
-					pxr::UsdAttribute ConeAngleAttr = ShapingAPI.GetShapingConeAngleAttr();
-					pxr::UsdAttribute ConeSoftnessAttr = ShapingAPI.GetShapingConeSoftnessAttr();
 
-					if (IntensityAttr && ExposureAttr && RadiusAttr && ConeAngleAttr && ConeSoftnessAttr)
+					if (IntensityAttr && ExposureAttr)
 					{
 						// Default values directly from pxr/usd/usdLux/schema.usda
 						float DefaultUsdIntensity = 1.0f;
 						float DefaultUsdExposure = 0.0f;
-						float DefaultUsdRadius = 0.5f;
-						float DefaultUsdConeAngle = 90.0f;
-						float DefaultUsdConeSoftness = 0.0f;
 
 						IntensityAttr.Get<float>(&DefaultUsdIntensity);
 						ExposureAttr.Get<float>(&DefaultUsdExposure);
-						RadiusAttr.Get<float>(&DefaultUsdRadius);
-						ConeAngleAttr.Get<float>(&DefaultUsdConeAngle);
-						ConeSoftnessAttr.Get<float>(&DefaultUsdConeSoftness);
 
-						float Default = UsdToUnreal::ConvertLuxShapingAPIIntensityAttr(
-							DefaultUsdIntensity,
-							DefaultUsdExposure,
-							DefaultUsdRadius,
-							DefaultUsdConeAngle,
-							DefaultUsdConeSoftness,
-							StageInfo
-						);
+						float Default = UsdToUnreal::ConvertLightIntensityAttr(DefaultUsdIntensity, DefaultUsdExposure);
 
-						Reader.FloatReader =
-							[IntensityAttr, ExposureAttr, RadiusAttr, ConeAngleAttr, ConeSoftnessAttr, Default, StageInfo](double UsdTimeCode)
+						Reader.FloatReader = [IntensityAttr, ExposureAttr, Default, StageInfo](double UsdTimeCode)
 						{
 							float Result = Default;
 
 							float UsdIntensity = 1.0f;
 							float UsdExposure = 0.0f;
-							float UsdRadius = 0.5f;
-							float UsdConeAngle = 90.0f;
-							float UsdConeSoftness = 0.0f;
-							if (IntensityAttr.Get<float>(&UsdIntensity, UsdTimeCode) && ExposureAttr.Get<float>(&UsdExposure, UsdTimeCode)
-								&& RadiusAttr.Get<float>(&UsdRadius, UsdTimeCode) && ConeAngleAttr.Get<float>(&UsdConeAngle, UsdTimeCode)
-								&& ConeSoftnessAttr.Get<float>(&UsdConeSoftness, UsdTimeCode))
+							if (IntensityAttr.Get<float>(&UsdIntensity, UsdTimeCode) && ExposureAttr.Get<float>(&UsdExposure, UsdTimeCode))
 							{
-								Result = UsdToUnreal::ConvertLuxShapingAPIIntensityAttr(
-									UsdIntensity,
-									UsdExposure,
-									UsdRadius,
-									UsdConeAngle,
-									UsdConeSoftness,
-									StageInfo
-								);
+								Result = UsdToUnreal::ConvertLightIntensityAttr(UsdIntensity, UsdExposure);
 							}
 
 							return Result;
@@ -1779,30 +1750,21 @@ UsdToUnreal::FPropertyTrackReader UsdToUnreal::CreatePropertyTrackReader(
 						// Default values directly from pxr/usd/usdLux/schema.usda
 						float DefaultUsdIntensity = 1.0f;
 						float DefaultUsdExposure = 0.0f;
-						float DefaultUsdRadius = 0.5f;
 
 						IntensityAttr.Get<float>(&DefaultUsdIntensity);
 						ExposureAttr.Get<float>(&DefaultUsdExposure);
-						RadiusAttr.Get<float>(&DefaultUsdRadius);
 
-						float Default = UsdToUnreal::ConvertSphereLightIntensityAttr(
-							DefaultUsdIntensity,
-							DefaultUsdExposure,
-							DefaultUsdRadius,
-							StageInfo
-						);
+						float Default = UsdToUnreal::ConvertLightIntensityAttr(DefaultUsdIntensity, DefaultUsdExposure);
 
-						Reader.FloatReader = [IntensityAttr, ExposureAttr, RadiusAttr, Default, StageInfo](double UsdTimeCode)
+						Reader.FloatReader = [IntensityAttr, ExposureAttr, Default, StageInfo](double UsdTimeCode)
 						{
 							float Result = Default;
 
 							float UsdIntensity = 1.0f;
 							float UsdExposure = 0.0f;
-							float UsdRadius = 0.5f;
-							if (IntensityAttr.Get<float>(&UsdIntensity, UsdTimeCode) && ExposureAttr.Get<float>(&UsdExposure, UsdTimeCode)
-								&& RadiusAttr.Get<float>(&UsdRadius, UsdTimeCode))
+							if (IntensityAttr.Get<float>(&UsdIntensity, UsdTimeCode) && ExposureAttr.Get<float>(&UsdExposure, UsdTimeCode))
 							{
-								Result = UsdToUnreal::ConvertSphereLightIntensityAttr(UsdIntensity, UsdExposure, UsdRadius, StageInfo);
+								Result = UsdToUnreal::ConvertLightIntensityAttr(UsdIntensity, UsdExposure);
 							}
 
 							return Result;
@@ -1860,42 +1822,27 @@ UsdToUnreal::FPropertyTrackReader UsdToUnreal::CreatePropertyTrackReader(
 			{
 				pxr::UsdAttribute IntensityAttr = RectLight.GetIntensityAttr();
 				pxr::UsdAttribute ExposureAttr = RectLight.GetExposureAttr();
-				pxr::UsdAttribute WidthAttr = RectLight.GetWidthAttr();
-				pxr::UsdAttribute HeightAttr = RectLight.GetHeightAttr();
 
-				if (IntensityAttr && ExposureAttr && WidthAttr && HeightAttr)
+				if (IntensityAttr && ExposureAttr)
 				{
 					// Default values directly from pxr/usd/usdLux/schema.usda
 					float DefaultUsdIntensity = 1.0f;
 					float DefaultUsdExposure = 0.0f;
-					float DefaultUsdWidth = 1.0f;
-					float DefaultUsdHeight = 1.0f;
 
 					IntensityAttr.Get<float>(&DefaultUsdIntensity);
 					ExposureAttr.Get<float>(&DefaultUsdExposure);
-					WidthAttr.Get<float>(&DefaultUsdWidth);
-					HeightAttr.Get<float>(&DefaultUsdHeight);
-
-					float Default = UsdToUnreal::ConvertRectLightIntensityAttr(
-						DefaultUsdIntensity,
-						DefaultUsdExposure,
-						DefaultUsdWidth,
-						DefaultUsdHeight,
-						StageInfo
-					);
 
-					Reader.FloatReader = [IntensityAttr, ExposureAttr, WidthAttr, HeightAttr, Default, StageInfo](double UsdTimeCode)
+					float Default = UsdToUnreal::ConvertLightIntensityAttr(DefaultUsdIntensity, DefaultUsdExposure);
+
+					Reader.FloatReader = [IntensityAttr, ExposureAttr, Default, StageInfo](double UsdTimeCode)
 					{
 						float Result = Default;
 
 						float UsdIntensity = 1.0f;
 						float UsdExposure = 0.0f;
-						float UsdWidth = 1.0f;
-						float UsdHeight = 1.0f;
-						if (IntensityAttr.Get<float>(&UsdIntensity, UsdTimeCode) && ExposureAttr.Get<float>(&UsdExposure, UsdTimeCode)
-							&& WidthAttr.Get<float>(&UsdWidth, UsdTimeCode) && HeightAttr.Get<float>(&UsdHeight, UsdTimeCode))
+						if (IntensityAttr.Get<float>(&UsdIntensity, UsdTimeCode) && ExposureAttr.Get<float>(&UsdExposure, UsdTimeCode))
 						{
-							Result = UsdToUnreal::ConvertRectLightIntensityAttr(UsdIntensity, UsdExposure, UsdWidth, UsdHeight, StageInfo);
+							Result = UsdToUnreal::ConvertLightIntensityAttr(UsdIntensity, UsdExposure);
 						}
 
 						return Result;
@@ -1933,32 +1880,27 @@ UsdToUnreal::FPropertyTrackReader UsdToUnreal::CreatePropertyTrackReader(
 			{
 				pxr::UsdAttribute IntensityAttr = DiskLight.GetIntensityAttr();
 				pxr::UsdAttribute ExposureAttr = DiskLight.GetExposureAttr();
-				pxr::UsdAttribute RadiusAttr = DiskLight.GetRadiusAttr();
 
-				if (IntensityAttr && ExposureAttr && RadiusAttr)
+				if (IntensityAttr && ExposureAttr)
 				{
 					// Default values directly from pxr/usd/usdLux/schema.usda
 					float DefaultUsdIntensity = 1.0f;
 					float DefaultUsdExposure = 0.0f;
-					float DefaultUsdRadius = 0.5f;
 
 					IntensityAttr.Get<float>(&DefaultUsdIntensity);
 					ExposureAttr.Get<float>(&DefaultUsdExposure);
-					RadiusAttr.Get<float>(&DefaultUsdRadius);
 
-					float Default = UsdToUnreal::ConvertDiskLightIntensityAttr(DefaultUsdIntensity, DefaultUsdExposure, DefaultUsdRadius, StageInfo);
+					float Default = UsdToUnreal::ConvertLightIntensityAttr(DefaultUsdIntensity, DefaultUsdExposure);
 
-					Reader.FloatReader = [IntensityAttr, ExposureAttr, RadiusAttr, Default, StageInfo](double UsdTimeCode)
+					Reader.FloatReader = [IntensityAttr, ExposureAttr, Default, StageInfo](double UsdTimeCode)
 					{
 						float Result = Default;
 
 						float UsdIntensity = 1.0f;
 						float UsdExposure = 0.0f;
-						float UsdRadius = 0.5f;
-						if (IntensityAttr.Get<float>(&UsdIntensity, UsdTimeCode) && ExposureAttr.Get<float>(&UsdExposure, UsdTimeCode)
-							&& RadiusAttr.Get<float>(&UsdRadius, UsdTimeCode))
+						if (IntensityAttr.Get<float>(&UsdIntensity, UsdTimeCode) && ExposureAttr.Get<float>(&UsdExposure, UsdTimeCode))
 						{
-							Result = UsdToUnreal::ConvertDiskLightIntensityAttr(UsdIntensity, UsdExposure, UsdRadius, StageInfo);
+							Result = UsdToUnreal::ConvertLightIntensityAttr(UsdIntensity, UsdExposure);
 						}
 
 						return Result;
@@ -1999,7 +1941,7 @@ UsdToUnreal::FPropertyTrackReader UsdToUnreal::CreatePropertyTrackReader(
 					IntensityAttr.Get<float>(&DefaultUsdIntensity);
 					ExposureAttr.Get<float>(&DefaultUsdExposure);
 
-					float Default = UsdToUnreal::ConvertDistantLightIntensityAttr(DefaultUsdIntensity, DefaultUsdExposure);
+					float Default = UsdToUnreal::ConvertLightIntensityAttr(DefaultUsdIntensity, DefaultUsdExposure);
 
 					Reader.FloatReader = [IntensityAttr, ExposureAttr, Default, StageInfo](double UsdTimeCode)
 					{
@@ -2009,7 +1951,7 @@ UsdToUnreal::FPropertyTrackReader UsdToUnreal::CreatePropertyTrackReader(
 						float UsdExposure = 0.0f;
 						if (IntensityAttr.Get<float>(&UsdIntensity, UsdTimeCode) && ExposureAttr.Get<float>(&UsdExposure, UsdTimeCode))
 						{
-							Result = UsdToUnreal::ConvertDistantLightIntensityAttr(UsdIntensity, UsdExposure);
+							Result = UsdToUnreal::ConvertLightIntensityAttr(UsdIntensity, UsdExposure);
 						}
 
 						return Result;
@@ -5110,7 +5052,7 @@ bool UnrealToUsd::CreateComponentPropertyBaker(
 				TFunction<const URectLightComponent*(void)> ComponentGetter = CreateComponentGetter(RectLightComponent);
 
 				TFunction<void(float, double)> WidthSetter = CreateCachedAttrSetter(WidthAttr, DistanceConversion);
-				TFunction<void(float, double)> HeightSetter = CreateCachedAttrSetter(WidthAttr, DistanceConversion);
+				TFunction<void(float, double)> HeightSetter = CreateCachedAttrSetter(HeightAttr, DistanceConversion);
 
 				BakerType = EBakingType::Light;
 				BakerFunction = [ComponentGetter,
@@ -5967,7 +5909,7 @@ UnrealToUsd::FPropertyTrackWriter UnrealToUsd::CreatePropertyTrackWriter(
 						}
 
 						// For now we'll assume the light intensity units are constant and the user doesn't have any light intensity unit tracks...
-						ELightUnits Units = ELightUnits::Lumens;
+						ELightUnits Units = ELightUnits::Nits;
 						if (const ULocalLightComponent* LightComponent = Cast<const ULocalLightComponent>(&Component))
 						{
 							Units = LightComponent->IntensityUnits;
@@ -6048,7 +5990,7 @@ UnrealToUsd::FPropertyTrackWriter UnrealToUsd::CreatePropertyTrackWriter(
 						}
 
 						// For now we'll assume the light intensity units are constant and the user doesn't have any light intensity unit tracks...
-						ELightUnits Units = ELightUnits::Lumens;
+						ELightUnits Units = ELightUnits::Nits;
 						if (const ULocalLightComponent* LightComponent = Cast<const ULocalLightComponent>(&Component))
 						{
 							Units = LightComponent->IntensityUnits;
@@ -6110,7 +6052,7 @@ UnrealToUsd::FPropertyTrackWriter UnrealToUsd::CreatePropertyTrackWriter(
 					}
 
 					// For now we'll assume the light intensity units are constant and the user doesn't have any light intensity unit tracks...
-					ELightUnits Units = ELightUnits::Lumens;
+					ELightUnits Units = ELightUnits::Nits;
 					if (const ULocalLightComponent* LightComponent = Cast<const ULocalLightComponent>(&Component))
 					{
 						Units = LightComponent->IntensityUnits;
@@ -6162,7 +6104,7 @@ UnrealToUsd::FPropertyTrackWriter UnrealToUsd::CreatePropertyTrackWriter(
 					}
 
 					// For now we'll assume the light intensity units are constant and the user doesn't have any light intensity unit tracks...
-					ELightUnits Units = ELightUnits::Lumens;
+					ELightUnits Units = ELightUnits::Nits;
 					if (const ULocalLightComponent* LightComponent = Cast<const ULocalLightComponent>(&Component))
 					{
 						Units = LightComponent->IntensityUnits;
@@ -6203,6 +6145,28 @@ UnrealToUsd::FPropertyTrackWriter UnrealToUsd::CreatePropertyTrackWriter(
 						ExposureAttr.Clear();
 					}
 
+					if (Attr)
+					{
+						Result.FloatWriter = [Attr](float UEValue, double UsdTimeCode)
+						{
+							Attr.Set(UnrealToUsd::ConvertLightIntensityProperty(UEValue), UsdTimeCode);
+						};
+					}
+				}
+			}
+			else if (pxr::UsdLuxDomeLight DomeLight{UsdPrim})
+			{
+				if (PropertyPath == UnrealIdentifiers::IntensityPropertyName)
+				{
+					Attr = DomeLight.GetIntensityAttr();
+
+					// Always clear exposure because we'll put all of our "light intensity" on the intensity attr and assume exposure
+					// is zero, as we can't manipulate something like that exposure directly from UE anyway
+					if (pxr::UsdAttribute ExposureAttr = DomeLight.GetExposureAttr())
+					{
+						ExposureAttr.Clear();
+					}
+
 					if (Attr)
 					{
 						Result.FloatWriter = [Attr](float UEValue, double UsdTimeCode)
diff --git a/Engine/Plugins/Runtime/USDCore/Source/USDUtilities/Public/USDLightConversion.h b/Engine/Plugins/Runtime/USDCore/Source/USDUtilities/Public/USDLightConversion.h
--- a/Engine/Plugins/Runtime/USDCore/Source/USDUtilities/Public/USDLightConversion.h
+++ b/Engine/Plugins/Runtime/USDCore/Source/USDUtilities/Public/USDLightConversion.h
@@ -100,7 +100,6 @@ namespace UsdToUnreal
 
 	// These are separately exposed so that they can be reused when reading data into MovieScene tracks.
 	// The other attribute conversions are mostly trivial like a single call to UsdToUnreal::ConvertDistance.
-	// Intensity return values are always in Lumen
 	USDUTILITIES_API float ConvertLightIntensityAttr(float UsdIntensity, float UsdExposure);
 	USDUTILITIES_API float ConvertDistantLightIntensityAttr(float UsdIntensity, float UsdExposure);
 	USDUTILITIES_API float ConvertRectLightIntensityAttr(
@@ -190,9 +189,7 @@ namespace UnrealToUsd
 
 namespace UsdUtils
 {
-	// As far as we can tell in Aug 2021, the best approximation for a standard in USD light units is to
-	// use Nits all the time. This doesn't make a lot of sense in some cases but here we try our best to convert to Nits
-	float ConvertIntensityToNits(float Intensity, float Steradians, float AreaInSqMeters, ELightUnits SourceUnits);
+	USDUTILITIES_API float ConvertIntensityToNits(float Intensity, float Steradians, float AreaInSqMeters, ELightUnits SourceUnits);
 }
 
 #endif	  // #if USE_USD_SDK
