USD: Delete code that was deprecated up to and including in 5.5.

#Synced-CL 49920387

diff --git a/Engine/Plugins/Importers/USDImporter/Source/GeometryCacheUSD/Private/GeometryCacheTrackUSD.cpp b/Engine/Plugins/Importers/USDImporter/Source/GeometryCacheUSD/Private/GeometryCacheTrackUSD.cpp
--- a/Engine/Plugins/Importers/USDImporter/Source/GeometryCacheUSD/Private/GeometryCacheTrackUSD.cpp
+++ b/Engine/Plugins/Importers/USDImporter/Source/GeometryCacheUSD/Private/GeometryCacheTrackUSD.cpp
@@ -283,19 +283,6 @@ void UGeometryCacheTrackUsd::Initialize(
 	UsdStream.Reset(new FGeometryCacheUsdStream(this, InReadFunc));
 }
 
-void UGeometryCacheTrackUsd::Initialize(
-	const UE::FUsdStage& InStage,
-	const FString& InPrimPath,
-	const FName& InRenderContext,
-	const TMap<FString, TMap<FString, int32>>& InMaterialToPrimvarToUVIndex,
-	int32 InStartFrameIndex,
-	int32 InEndFrameIndex,
-	FReadUsdMeshFunction InReadFunc
-)
-{
-	Initialize(InStage, InPrimPath, InStartFrameIndex, InEndFrameIndex, InReadFunc);
-}
-
 void UGeometryCacheTrackUsd::UpdateTime(float Time, bool bLooping)
 {
 	if (UsdStream)
diff --git a/Engine/Plugins/Importers/USDImporter/Source/GeometryCacheUSD/Public/GeometryCacheTrackUSD.h b/Engine/Plugins/Importers/USDImporter/Source/GeometryCacheUSD/Public/GeometryCacheTrackUSD.h
--- a/Engine/Plugins/Importers/USDImporter/Source/GeometryCacheUSD/Public/GeometryCacheTrackUSD.h
+++ b/Engine/Plugins/Importers/USDImporter/Source/GeometryCacheUSD/Public/GeometryCacheTrackUSD.h
@@ -34,17 +34,6 @@ public:
 		FReadUsdMeshFunction InReadFunc
 	);
 
-	UE_DEPRECATED(5.3, "The RenderContext and InMaterialToPrimvarToUVIndex parameters are no longer used.")
-	UE_API void Initialize(
-		const UE::FUsdStage& InStage,
-		const FString& InPrimPath,
-		const FName& InRenderContext,
-		const TMap<FString, TMap<FString, int32>>& InMaterialToPrimvarToUVIndex,
-		int32 InStartFrameIndex,
-		int32 InEndFrameIndex,
-		FReadUsdMeshFunction InReadFunc
-	);
-
 	//~ Begin UObject Interface.
 	UE_API virtual void BeginDestroy() override;
 	UE_API virtual void GetResourceSizeEx(FResourceSizeEx& CumulativeResourceSize) override;
diff --git a/Engine/Plugins/Importers/USDImporter/Source/USDClassesEditor/Private/USDClassesEditorModule.cpp b/Engine/Plugins/Importers/USDImporter/Source/USDClassesEditor/Private/USDClassesEditorModule.cpp
--- a/Engine/Plugins/Importers/USDImporter/Source/USDClassesEditor/Private/USDClassesEditorModule.cpp
+++ b/Engine/Plugins/Importers/USDImporter/Source/USDClassesEditor/Private/USDClassesEditorModule.cpp
@@ -2,36 +2,12 @@
 
 #include "USDClassesEditorModule.h"
 
-#include "USDAssetCache2.h"
-
 #include "Modules/ModuleManager.h"
 
 #define LOCTEXT_NAMESPACE "USDClassesEditorModule"
 
-UUsdAssetCache2* IUsdClassesEditorModule::ShowMissingDefaultAssetCacheDialog()
-{
-	return nullptr;
-}
-
-void IUsdClassesEditorModule::ShowMissingDefaultAssetCacheDialog(UUsdAssetCache2*& OutCreatedCache, bool& bOutUserAccepted)
-{
-}
-
-EDefaultAssetCacheDialogOption IUsdClassesEditorModule::ShowMissingDefaultAssetCacheDialog(UUsdAssetCache2*& OutCreatedCache)
-{
-	return EDefaultAssetCacheDialogOption::Cancel;
-}
-
 class FUsdClassesEditorModule : public IUsdClassesEditorModule
 {
-public:
-	virtual void StartupModule() override
-	{
-	}
-
-	virtual void ShutdownModule() override
-	{
-	}
 };
 
 IMPLEMENT_MODULE(FUsdClassesEditorModule, USDClassesEditor);
diff --git a/Engine/Plugins/Importers/USDImporter/Source/USDClassesEditor/Public/USDClassesEditorModule.h b/Engine/Plugins/Importers/USDImporter/Source/USDClassesEditor/Public/USDClassesEditorModule.h
--- a/Engine/Plugins/Importers/USDImporter/Source/USDClassesEditor/Public/USDClassesEditorModule.h
+++ b/Engine/Plugins/Importers/USDImporter/Source/USDClassesEditor/Public/USDClassesEditorModule.h
@@ -4,26 +4,6 @@
 
 #include "Modules/ModuleInterface.h"
 
-class UUsdAssetCache2;
-
-enum class EDefaultAssetCacheDialogOption : uint8
-{
-	Cancel = 0,
-	PickExisting,
-	CreateNew,
-	DontUseDefault
-};
-
 class IUsdClassesEditorModule : public IModuleInterface
 {
-public:
-	/** Shows the dialog that asks the user to set a default asset cache for the project */
-	UE_DEPRECATED(5.3, "Use the other signature that also returns whether the user accepted the dialog or not")
-	USDCLASSESEDITOR_API static UUsdAssetCache2* ShowMissingDefaultAssetCacheDialog();
-
-	UE_DEPRECATED(5.4, "Use the other signature that also returns the dialog outcome as an Enum")
-	USDCLASSESEDITOR_API static void ShowMissingDefaultAssetCacheDialog(UUsdAssetCache2*& OutCreatedCache, bool& bOutUserAccepted);
-
-	UE_DEPRECATED(5.5, "The Default Asset Cache dialog has been deprecated. An asset cache for the project will always be created on-demand now")
-	USDCLASSESEDITOR_API static EDefaultAssetCacheDialogOption ShowMissingDefaultAssetCacheDialog(UUsdAssetCache2*& OutCreatedCache);
 };
diff --git a/Engine/Plugins/Importers/USDImporter/Source/USDExporter/Private/MaterialExporterUSD.cpp b/Engine/Plugins/Importers/USDImporter/Source/USDExporter/Private/MaterialExporterUSD.cpp
--- a/Engine/Plugins/Importers/USDImporter/Source/USDExporter/Private/MaterialExporterUSD.cpp
+++ b/Engine/Plugins/Importers/USDImporter/Source/USDExporter/Private/MaterialExporterUSD.cpp
@@ -403,28 +403,6 @@ bool UMaterialExporterUsd::ExportMaterial(
 #endif	  // USE_USD_SDK
 }
 
-// Deprecated signature
-bool UMaterialExporterUsd::ExportMaterial(
-	const UMaterialInterface& Material,
-	const FUsdMaterialBakingOptions& Options,
-	const FFilePath& FilePath,
-	bool bReplaceIdentical,
-	bool bReExportIdenticalAssets,
-	bool bIsAutomated
-)
-{
-	static FUsdMetadataExportOptions DefaultOptions;
-	return UMaterialExporterUsd::ExportMaterial(
-		Material,
-		Options,
-		DefaultOptions,
-		FilePath,
-		bReplaceIdentical,
-		bReExportIdenticalAssets,
-		bIsAutomated
-	);
-}
-
 bool UMaterialExporterUsd::ExportMaterialsForStage(
 	const TArray<UMaterialInterface*>& Materials,
 	const FUsdMaterialBakingOptions& Options,
@@ -530,30 +508,4 @@ bool UMaterialExporterUsd::ExportMaterialsForStage(
 #endif	  // USE_USD_SDK
 }
 
-// Deprecated signature
-bool UMaterialExporterUsd::ExportMaterialsForStage(
-	const TArray<UMaterialInterface*>& Materials,
-	const FUsdMaterialBakingOptions& Options,
-	const FString& StageRootLayerPath,
-	bool bIsAssetLayer,
-	bool bUsePayload,
-	bool bReplaceIdentical,
-	bool bReExportIdenticalAssets,
-	bool bIsAutomated
-)
-{
-	static FUsdMetadataExportOptions DefaultOptions;
-	return UMaterialExporterUsd::ExportMaterialsForStage(
-		Materials,
-		Options,
-		DefaultOptions,
-		StageRootLayerPath,
-		bIsAssetLayer,
-		bUsePayload,
-		bReplaceIdentical,
-		bReExportIdenticalAssets,
-		bIsAutomated
-	);
-}
-
 #undef LOCTEXT_NAMESPACE
diff --git a/Engine/Plugins/Importers/USDImporter/Source/USDExporter/Private/USDConversionBlueprintContext.cpp b/Engine/Plugins/Importers/USDImporter/Source/USDExporter/Private/USDConversionBlueprintContext.cpp
--- a/Engine/Plugins/Importers/USDImporter/Source/USDExporter/Private/USDConversionBlueprintContext.cpp
+++ b/Engine/Plugins/Importers/USDImporter/Source/USDExporter/Private/USDConversionBlueprintContext.cpp
@@ -570,30 +570,7 @@ void UUsdConversionBlueprintContext::ReplaceUnrealMaterialsWithBaked(
 
 bool UUsdConversionBlueprintContext::RemoveUnrealSurfaceOutput(const FString& PrimPath, const FFilePath& LayerToAuthorIn)
 {
-#if USE_USD_SDK
-	if (!Stage)
-	{
-		return false;
-	}
-
-	UE::FUsdPrim Prim = UnrealToUsdImpl::GetPrim(Stage, PrimPath);
-	if (!Prim)
-	{
-		return false;
-	}
-
-	UE::FSdfLayer Layer;
-	if (!LayerToAuthorIn.FilePath.IsEmpty())
-	{
-		Layer = UE::FSdfLayer::FindOrOpen(*LayerToAuthorIn.FilePath);
-	}
-
-	PRAGMA_DISABLE_DEPRECATION_WARNINGS
-	return UsdUtils::RemoveUnrealSurfaceOutput(Prim, Layer);
-	PRAGMA_ENABLE_DEPRECATION_WARNINGS
-#else
 	return false;
-#endif	  // USE_USD_SDK
 }
 
 int32 UUsdConversionBlueprintContext::GetUsdStageNumFrames()
diff --git a/Engine/Plugins/Importers/USDImporter/Source/USDExporter/Public/MaterialExporterUSD.h b/Engine/Plugins/Importers/USDImporter/Source/USDExporter/Public/MaterialExporterUSD.h
--- a/Engine/Plugins/Importers/USDImporter/Source/USDExporter/Public/MaterialExporterUSD.h
+++ b/Engine/Plugins/Importers/USDImporter/Source/USDExporter/Public/MaterialExporterUSD.h
@@ -48,15 +48,6 @@ public:
 		bool bReExportIdenticalAssets = false,
 		bool bIsAutomated = false
 	);
-	UE_DEPRECATED(5.4, "Please use the function overload that can receive metadata options.")
-	static bool ExportMaterial(
-		const UMaterialInterface& Material,
-		const FUsdMaterialBakingOptions& Options,
-		const FFilePath& FilePath,
-		bool bReplaceIdentical = true,
-		bool bReExportIdenticalAssets = false,
-		bool bIsAutomated = false
-	);
 
 	/**
 	 * Exports the provided materials next to the stage's root layer on disk, replacing usages of unrealMaterials within `Stage` with references
@@ -88,15 +79,4 @@ public:
 		bool bReExportIdenticalAssets = false,
 		bool bIsAutomated = false
 	);
-	UE_DEPRECATED(5.4, "Please use the function overload that can receive metadata options.")
-	static bool ExportMaterialsForStage(
-		const TArray<UMaterialInterface*>& Materials,
-		const FUsdMaterialBakingOptions& Options,
-		const FString& StageRootLayerPath,
-		bool bIsAssetLayer,
-		bool bUsePayload,
-		bool bReplaceIdentical = true,
-		bool bReExportIdenticalAssets = false,
-		bool bIsAutomated = false
-	);
 };
diff --git a/Engine/Plugins/Importers/USDImporter/Source/USDExporter/Public/USDAssetOptions.h b/Engine/Plugins/Importers/USDImporter/Source/USDExporter/Public/USDAssetOptions.h
--- a/Engine/Plugins/Importers/USDImporter/Source/USDExporter/Public/USDAssetOptions.h
+++ b/Engine/Plugins/Importers/USDImporter/Source/USDExporter/Public/USDAssetOptions.h
@@ -69,9 +69,9 @@ struct FUsdMeshAssetOptions
 	bool bBakeMaterials = false;
 
 	/** Whether to remove the 'unrealMaterial' attribute after binding the corresponding baked material */
-	UE_DEPRECATED(5.2, "This option is now deprecated as UE material assignments are only visible in the 'unreal' render context anyway")
+	UE_DEPRECATED(all, "This option is now deprecated as UE material assignments are only visible in the 'unreal' render context anyway")
 	UPROPERTY(config, meta = (EditCondition = bBakeMaterials))
-	bool bRemoveUnrealMaterials = false;
+	bool bRemoveUnrealMaterials_DEPRECATED = false;
 
 	UPROPERTY(EditAnywhere, config, BlueprintReadWrite, Category = "Material options", meta = (EditCondition = bBakeMaterials))
 	FUsdMaterialBakingOptions MaterialBakingOptions;
diff --git a/Engine/Plugins/Importers/USDImporter/Source/USDExporter/Public/USDConversionBlueprintContext.h b/Engine/Plugins/Importers/USDImporter/Source/USDExporter/Public/USDConversionBlueprintContext.h
--- a/Engine/Plugins/Importers/USDImporter/Source/USDExporter/Public/USDConversionBlueprintContext.h
+++ b/Engine/Plugins/Importers/USDImporter/Source/USDExporter/Public/USDConversionBlueprintContext.h
@@ -202,16 +202,8 @@ public:
 		bool bRemoveUnrealMaterials = false
 	);
 
-	/**
-	 * Clears any opinions for the 'unreal' render context surface output of MaterialPrim within LayerToAuthorIn.
-	 * If LayerToAuthorIn is empty it will clear opinions from all layers of the stage's layer stack.
-	 *
-	 * @param Primpath - Path to the prim pxr::UsdShadeMaterial schema to update the 'unreal' surface output of (e.g. "/Root/MyCube/Red")
-	 * @param LayerToAuthorIn - Layer to clear the opinions in. Can be the empty string to clear opinions from all layers in the layer stack
-	 * @return Whether we successfully cleared the opinions or not
-	 */
-	UE_DEPRECATED(5.2, "No longer used as UE material assignments are only visible in the 'unreal' render context anyway")
-	UFUNCTION(BlueprintCallable, Category = "Conversion utils")
+	UE_DEPRECATED(all, "No longer used as UE material assignments are only visible in the 'unreal' render context anyway")
+	UFUNCTION()
 	UE_API bool RemoveUnrealSurfaceOutput(const FString& PrimPath, const FFilePath& LayerToAuthorIn);
 
 	/** Returns how many frames of animation the stage would generate if imported */
diff --git a/Engine/Plugins/Importers/USDImporter/Source/USDSchemas/Private/USDSchemasModule.cpp b/Engine/Plugins/Importers/USDImporter/Source/USDSchemas/Private/USDSchemasModule.cpp
--- a/Engine/Plugins/Importers/USDImporter/Source/USDSchemas/Private/USDSchemasModule.cpp
+++ b/Engine/Plugins/Importers/USDImporter/Source/USDSchemas/Private/USDSchemasModule.cpp
@@ -103,18 +103,7 @@ public:
 		return FUsdSchemaTranslatorRegistry::Get();
 	}
 
-	PRAGMA_DISABLE_DEPRECATION_WARNINGS
-	virtual FUsdRenderContextRegistry& GetRenderContextRegistry() override
-	{
-		return UsdRenderContextRegistry;
-	}
-	PRAGMA_ENABLE_DEPRECATION_WARNINGS
-
 protected:
-	PRAGMA_DISABLE_DEPRECATION_WARNINGS
-	FUsdRenderContextRegistry UsdRenderContextRegistry;
-	PRAGMA_ENABLE_DEPRECATION_WARNINGS
-
 	TArray<FRegisteredSchemaTranslatorHandle> TranslatorHandles;
 };
 
diff --git a/Engine/Plugins/Importers/USDImporter/Source/USDSchemas/Private/USDSkelRootTranslator.cpp b/Engine/Plugins/Importers/USDImporter/Source/USDSchemas/Private/USDSkelRootTranslator.cpp
deleted file mode 100644
diff --git a/Engine/Plugins/Importers/USDImporter/Source/USDSchemas/Private/USDSkelSkeletonTranslator.cpp b/Engine/Plugins/Importers/USDImporter/Source/USDSchemas/Private/USDSkelSkeletonTranslator.cpp
--- a/Engine/Plugins/Importers/USDImporter/Source/USDSchemas/Private/USDSkelSkeletonTranslator.cpp
+++ b/Engine/Plugins/Importers/USDImporter/Source/USDSchemas/Private/USDSkelSkeletonTranslator.cpp
@@ -2025,20 +2025,6 @@ void FUsdSkelSkeletonTranslator::UpdateComponents(USceneComponent* SceneComponen
 	{
 		PrimWithLiveLinkSchema = SkeletonPrim;
 	}
-	else if (UsdUtils::PrimHasSchema(ClosestSkelRootPrim, UnrealIdentifiers::LiveLinkAPI))
-	{
-		PrimWithLiveLinkSchema = ClosestSkelRootPrim;
-
-		// Commenting the usual deprecation macro so that we can find this with search and replace later
-		// UE_DEPRECATED(5.4, "schemas")
-		USD_LOG_USERWARNING(FText::Format(
-			LOCTEXT(
-				"DeprecatedSchemas",
-				"Placing integration schemas (Live Link, Control Rig, Groom Binding) on SkelRoot prims (like '{0}') has been deprecated on version 5.4 and will be unsupported in a future release. Please place your integration schemas directly on the Skeleton prims instead!"
-			),
-			FText::FromString(ClosestSkelRootPrim.GetPrimPath().GetString())
-		));
-	}
 	const bool bPrimHasLiveLinkSchema = static_cast<bool>(PrimWithLiveLinkSchema);
 
 	bool bPrimHasLiveLinkEnabled = bPrimHasLiveLinkSchema;
@@ -2187,24 +2173,9 @@ void FUsdSkelSkeletonTranslator::UpdateComponents(USceneComponent* SceneComponen
 	}
 
 	// If the prim has a GroomBinding schema, apply the target groom to its associated GroomComponent
+	if (UsdUtils::PrimHasSchema(SkeletonPrim, UnrealIdentifiers::GroomBindingAPI))
 	{
-		if (UsdUtils::PrimHasSchema(SkeletonPrim, UnrealIdentifiers::GroomBindingAPI))
-		{
-			UsdGroomTranslatorUtils::SetGroomFromPrim(SkeletonPrim, *Context->PrimLinkCache, SceneComponent);
-		}
-		else if (UsdUtils::PrimHasSchema(ClosestSkelRootPrim, UnrealIdentifiers::GroomBindingAPI))
-		{
-			// Commenting the usual deprecation macro so that we can find this with search and replace later
-			// UE_DEPRECATED(5.4, "schemas")
-			USD_LOG_USERWARNING(FText::Format(
-				LOCTEXT(
-					"DeprecatedSchemas",
-					"Placing integration schemas (Live Link, Control Rig, Groom Binding) on SkelRoot prims (like '{0}') has been deprecated on version 5.4 and will be unsupported in a future release. Please place your integration schemas directly on the Skeleton prims instead!"
-				),
-				FText::FromString(ClosestSkelRootPrim.GetPrimPath().GetString())
-			));
-			UsdGroomTranslatorUtils::SetGroomFromPrim(ClosestSkelRootPrim, *Context->PrimLinkCache, SceneComponent);
-		}
+		UsdGroomTranslatorUtils::SetGroomFromPrim(SkeletonPrim, *Context->PrimLinkCache, SceneComponent);
 	}
 #endif	  // WITH_EDITOR
 }
diff --git a/Engine/Plugins/Importers/USDImporter/Source/USDSchemas/Public/USDSchemasModule.h b/Engine/Plugins/Importers/USDImporter/Source/USDSchemas/Public/USDSchemasModule.h
--- a/Engine/Plugins/Importers/USDImporter/Source/USDSchemas/Public/USDSchemasModule.h
+++ b/Engine/Plugins/Importers/USDImporter/Source/USDSchemas/Public/USDSchemasModule.h
@@ -15,10 +15,6 @@ class IUsdSchemasModule : public IModuleInterface
 public:
 	UE_DEPRECATED(5.6, "Use FUsdSchemaTranslatorRegistry::Get() from Objects/USDSchemaTranslator.h file in the USDUtilities module (USDCore plugin).")
 	virtual FUsdSchemaTranslatorRegistry& GetTranslatorRegistry() = 0;
-
-	PRAGMA_DISABLE_DEPRECATION_WARNINGS
-	virtual FUsdRenderContextRegistry& GetRenderContextRegistry() = 0;
-	PRAGMA_ENABLE_DEPRECATION_WARNINGS
 };
 
 namespace UsdUnreal::Analytics
diff --git a/Engine/Plugins/Importers/USDImporter/Source/USDSchemas/Public/USDSkelRootTranslator.h b/Engine/Plugins/Importers/USDImporter/Source/USDSchemas/Public/USDSkelRootTranslator.h
--- a/Engine/Plugins/Importers/USDImporter/Source/USDSchemas/Public/USDSkelRootTranslator.h
+++ b/Engine/Plugins/Importers/USDImporter/Source/USDSchemas/Public/USDSkelRootTranslator.h
@@ -2,26 +2,4 @@
 
 #pragma once
 
-#include "USDGeomMeshTranslator.h"
-
-#if USE_USD_SDK
-
-class USDSCHEMAS_API UE_DEPRECATED(5.4, "Use the UsdSkelSkeletonTranslator for skeletal data") FUsdSkelRootTranslator
-	: public FUsdGeomXformableTranslator
-{
-	using Super = FUsdGeomXformableTranslator;
-
-public:
-	using FUsdGeomXformableTranslator::FUsdGeomXformableTranslator;
-
-	virtual void CreateAssets() override;
-	virtual USceneComponent* CreateComponents() override;
-	virtual void UpdateComponents(USceneComponent* SceneComponent) override;
-
-	virtual bool CollapsesChildren(ECollapsingType CollapsingType) const override;
-	virtual bool CanBeCollapsed(ECollapsingType CollapsingType) const override;
-
-	virtual TSet<UE::FSdfPath> CollectAuxiliaryPrims() const override;
-};
-
-#endif	  // #if USE_USD_SDK
+UE_DEPRECATED_HEADER(5.8, "Use USDSkelSkeletonTranslator.h instead")
diff --git a/Engine/Plugins/Importers/USDImporter/Source/USDStage/Private/USDLevelSequenceHelper.cpp b/Engine/Plugins/Importers/USDImporter/Source/USDStage/Private/USDLevelSequenceHelper.cpp
--- a/Engine/Plugins/Importers/USDImporter/Source/USDStage/Private/USDLevelSequenceHelper.cpp
+++ b/Engine/Plugins/Importers/USDImporter/Source/USDStage/Private/USDLevelSequenceHelper.cpp
@@ -5455,12 +5455,6 @@ void FUsdLevelSequenceHelper::OnStageActorRenamed()
 	}
 }
 
-PRAGMA_DISABLE_DEPRECATION_WARNINGS
-void FUsdLevelSequenceHelper::SetInfoCache(TSharedPtr<FUsdInfoCache> InInfoCache)
-{
-}
-PRAGMA_ENABLE_DEPRECATION_WARNINGS
-
 void FUsdLevelSequenceHelper::SetPrimLinkCache(UUsdPrimLinkCache* PrimLinkCache)
 {
 	if (UsdSequencerImpl.IsValid())
diff --git a/Engine/Plugins/Importers/USDImporter/Source/USDStage/Private/USDStageActor.cpp b/Engine/Plugins/Importers/USDImporter/Source/USDStage/Private/USDStageActor.cpp
--- a/Engine/Plugins/Importers/USDImporter/Source/USDStage/Private/USDStageActor.cpp
+++ b/Engine/Plugins/Importers/USDImporter/Source/USDStage/Private/USDStageActor.cpp
@@ -2315,22 +2315,6 @@ UUsdPrimTwin* AUsdStageActor::ExpandPrim(
 		{
 			PrimWithSchema = Prim;
 		}
-		else if (UE::FUsdPrim ParentSkelRoot = UsdUtils::GetClosestParentSkelRoot(Prim))
-		{
-			if (UsdUtils::PrimHasSchema(ParentSkelRoot, UnrealIdentifiers::ControlRigAPI))
-			{
-				// Commenting the usual deprecation macro so that we can find this with search and replace later
-				// UE_DEPRECATED(5.4, "schemas")
-				USD_LOG_USERWARNING(FText::Format(
-					LOCTEXT(
-						"DeprecatedSchemas",
-						"Placing integration schemas (Live Link, Control Rig, Groom Binding) on SkelRoot prims (like '{0}') has been deprecated on version 5.4 and will be unsupported in a future release. Please place your integration schemas directly on the Skeleton prims instead!"
-					),
-					FText::FromString(Prim.GetPrimPath().GetString())
-				));
-				PrimWithSchema = ParentSkelRoot;
-			}
-		}
 
 		if (PrimWithSchema)
 		{
@@ -2440,13 +2424,6 @@ void AUsdStageActor::SetUsdStage(const UE::FUsdStage& NewStage)
 	OnStageChanged.Broadcast();
 }
 
-PRAGMA_DISABLE_DEPRECATION_WARNINGS
-UE::FUsdStage& AUsdStageActor::GetOrLoadUsdStage()
-{
-	return GetOrOpenUsdStage();
-}
-PRAGMA_ENABLE_DEPRECATION_WARNINGS
-
 UE::FUsdStage& AUsdStageActor::GetOrOpenUsdStage()
 {
 	OpenUsdStage();
@@ -2527,52 +2504,6 @@ void AUsdStageActor::SetStageState(EUsdStageState NewState)
 PRAGMA_DISABLE_DEPRECATION_WARNINGS
 void AUsdStageActor::SetAssetCache(UUsdAssetCache2* NewCache)
 {
-	if (NewCache == UsdAssetCache)
-	{
-		return;
-	}
-
-	const bool bMarkDirty = false;
-	Modify(bMarkDirty);
-
-	// Remove ourselves from our previous cache. We're going to have to create new assets anyway, so it doesn't matter
-	// if we discard our current assets.
-	if (UsdAssetCache)
-	{
-		UsdAssetCache->RemoveAllAssetReferences(this);
-		UsdAssetCache->RefreshStorage();
-	}
-
-	UsdAssetCache = NewCache;
-
-	// We can't have no cache while we have a stage loaded, so at least revert the property to a transient cache
-	// instead, as the intent may have been to just have the actor not point at the previous cache anymore.
-	if (!UsdAssetCache && UsdStage)
-	{
-		FNotificationInfo Toast(LOCTEXT("MustHaveCache", "Must have an Asset Cache"));
-		Toast.SubText = LOCTEXT(
-			"MustHaveCache_Subtext",
-			"The Stage Actor must always have an Asset Cache while a stage is loaded, so a temporary cache will be created.\n\nClose the stage "
-			"before clearing the cache if you wish to clear this property."
-		);
-		Toast.Image = FCoreStyle::Get().GetBrush(TEXT("MessageLog.Warning"));
-		Toast.bUseLargeFont = false;
-		Toast.bFireAndForget = true;
-		Toast.FadeOutDuration = 1.0f;
-		Toast.ExpireDuration = 12.0f;
-		Toast.bUseThrobber = false;
-		Toast.bUseSuccessFailIcons = false;
-		FSlateNotificationManager::Get().AddNotification(Toast);
-
-		UsdAssetCache = NewObject<UUsdAssetCache2>(GetTransientPackage(), NAME_None, GetMaskedFlags(RF_PropagateToSubObjects));
-	}
-
-	// Here we pretend we just received a root resync so that we re-fetch assets from the cache
-	// and update its components
-	UsdUtils::FObjectChangesByPath InfoChanges;
-	UsdUtils::FObjectChangesByPath ResyncChanges;
-	ResyncChanges.Add({TEXT("/"), {}});
-	OnUsdObjectsChanged(InfoChanges, ResyncChanges);
 }
 PRAGMA_ENABLE_DEPRECATION_WARNINGS
 
@@ -2711,16 +2642,6 @@ void AUsdStageActor::SetShareAssetsForIdenticalPrims(bool bShare)
 PRAGMA_DISABLE_DEPRECATION_WARNINGS
 void AUsdStageActor::SetReuseIdenticalAssets(bool bReuse)
 {
-	if (bReuse == bReuseIdenticalAssets)
-	{
-		return;
-	}
-
-	const bool bMarkDirty = false;
-	Modify(bMarkDirty);
-
-	bReuseIdenticalAssets = bReuse;
-	LoadUsdStage();
 }
 PRAGMA_ENABLE_DEPRECATION_WARNINGS
 
@@ -3578,13 +3499,6 @@ void AUsdStageActor::Refresh() const
 	OnTimeChanged.Broadcast();
 }
 
-PRAGMA_DISABLE_DEPRECATION_WARNINGS
-void AUsdStageActor::ReloadAnimations()
-{
-	RegenerateLevelSequence();
-}
-PRAGMA_ENABLE_DEPRECATION_WARNINGS
-
 bool AUsdStageActor::RegenerateLevelSequence()
 {
 	TRACE_CPUPROFILER_EVENT_SCOPE(AUsdStageActor::RegenerateLevelSequence);
@@ -3677,13 +3591,6 @@ void AUsdStageActor::OpenLevelSequence()
 #endif	  // WITH_EDITOR
 }
 
-PRAGMA_DISABLE_DEPRECATION_WARNINGS
-TSharedPtr<FUsdInfoCache> AUsdStageActor::GetInfoCache()
-{
-	return InfoCache;
-}
-PRAGMA_ENABLE_DEPRECATION_WARNINGS
-
 TSharedPtr<UE::FUsdGeomBBoxCache> AUsdStageActor::GetBBoxCache()
 {
 	return BBoxCache;
@@ -4887,14 +4794,6 @@ void AUsdStageActor::HandlePropertyChangedEvent(FPropertyChangedEvent& PropertyC
 		MetadataOptions.bInvertFilters = !bCorrectInvertValue;
 		SetInvertFilters(bCorrectInvertValue);
 	}
-	PRAGMA_DISABLE_DEPRECATION_WARNINGS
-	else if (PropertyName == GET_MEMBER_NAME_CHECKED(AUsdStageActor, UsdAssetCache))
-	{
-		UUsdAssetCache2* CorrectCache = UsdAssetCache.Get();
-		UsdAssetCache = UsdAssetCache ? nullptr : NewObject<UUsdAssetCache2>();
-		SetAssetCache(CorrectCache);
-	}
-	PRAGMA_ENABLE_DEPRECATION_WARNINGS
 	else if (PropertyName == GET_MEMBER_NAME_CHECKED(AUsdStageActor, AssetCache))
 	{
 		UUsdAssetCache3* CorrectCache = AssetCache.Get();
diff --git a/Engine/Plugins/Importers/USDImporter/Source/USDStage/Public/USDLevelSequenceHelper.h b/Engine/Plugins/Importers/USDImporter/Source/USDStage/Public/USDLevelSequenceHelper.h
--- a/Engine/Plugins/Importers/USDImporter/Source/USDStage/Public/USDLevelSequenceHelper.h
+++ b/Engine/Plugins/Importers/USDImporter/Source/USDStage/Public/USDLevelSequenceHelper.h
@@ -46,13 +46,7 @@ public:
 	/** Allows serialization for transaction support */
 	UE_API bool Serialize(FArchive& Ar);
 
-	PRAGMA_DISABLE_DEPRECATION_WARNINGS
-	/** Sets the asset cache to use when fetching assets and asset info required for the level sequence animation, like UAnimSequences */
-	UE_DEPRECATED(5.5, "Use SetPrimLinkCache instead")
-	UE_API void SetInfoCache(TSharedPtr<FUsdInfoCache> InInfoCache);
-	UE_API PRAGMA_ENABLE_DEPRECATION_WARNINGS
-
-	void SetPrimLinkCache(UUsdPrimLinkCache* PrimLinkCache);
+	UE_API void SetPrimLinkCache(UUsdPrimLinkCache* PrimLinkCache);
 
 	/** Sets the BBoxCache to use when importing bound animations for prims. Needed for importing, where we don't have a stage actor to take the
 	 * BBoxCache from */
diff --git a/Engine/Plugins/Importers/USDImporter/Source/USDStage/Public/USDStageActor.h b/Engine/Plugins/Importers/USDImporter/Source/USDStage/Public/USDStageActor.h
--- a/Engine/Plugins/Importers/USDImporter/Source/USDStage/Public/USDStageActor.h
+++ b/Engine/Plugins/Importers/USDImporter/Source/USDStage/Public/USDStageActor.h
@@ -82,8 +82,11 @@ public:
 	UPROPERTY(Transient, Instanced)
 	TObjectPtr<UUsdPrimLinkCache> PrimLinkCache;
 
+	PRAGMA_DISABLE_DEPRECATION_WARNINGS
+	UE_DEPRECATED(5.8, "Use the 'AssetCache' Property, that uses the new AssetCache3 asset type.")
 	UPROPERTY()
-	TObjectPtr<UUsdAssetCache2> UsdAssetCache;
+	TObjectPtr<UUsdAssetCache2> UsdAssetCache_DEPRECATED;
+	PRAGMA_ENABLE_DEPRECATION_WARNINGS
 
 	// These properties are configs so that spawned actors read them from the CDO when spawned.
 	// This allows the defaults for them to be configured on EditorPerProjectUserSettings.ini, and allows us to write
@@ -127,18 +130,13 @@ public:
 	UPROPERTY(EditAnywhere, BlueprintReadOnly, Category = "USD", config)
 	bool bShareAssetsForIdenticalPrims;
 
-	UE_DEPRECATED(5.5, "This property has been renamed to 'Share Assets for Identical Prims'")
+	UE_DEPRECATED(all, "This property has been renamed to 'Share Assets for Identical Prims'")
 	UPROPERTY()
-	bool bReuseIdenticalAssets;
+	bool bReuseIdenticalAssets_DEPRECATED;
 
-	/**
-	 * If true, will cause us to collapse any point instancer prim into a single static mesh and static mesh component.
-	 * If false, will cause us to use HierarchicalInstancedStaticMeshComponents to replicate the instancing behavior.
-	 * Point instancers inside other point instancer prototypes are *always* collapsed into the prototype's static mesh.
-	 */
-	UE_DEPRECATED(5.2, "This option is now controlled via the cvar 'USD.CollapseTopLevelPointInstancers'")
+	UE_DEPRECATED(all, "This option is now controlled via the cvar 'USD.CollapseTopLevelPointInstancers'")
 	UPROPERTY()
-	bool bCollapseTopLevelPointInstancers;
+	bool bCollapseTopLevelPointInstancers_DEPRECATED;
 
 	/* Only load prims with these specific purposes from the USD file */
 	UPROPERTY(EditAnywhere, BlueprintReadOnly, Category = "USD", config, meta = (Bitmask, BitmaskEnum = "/Script/UnrealUSDWrapper.EUsdPurpose"))
@@ -202,9 +200,11 @@ public:
 	UFUNCTION(BlueprintCallable, Category = "USD", meta = (CallInEditor = "true"))
 	USDSTAGE_API void SetStageState(EUsdStageState NewStageState);
 
-	UE_DEPRECATED(5.5, "The UUsdAssetCache2 class has been deprecated. Use UUsdAssetCache3 and call SetUsdAssetCache instead")
-	UFUNCTION(BlueprintCallable, Category = "USD", meta = (CallInEditor = "true"))
+	PRAGMA_DISABLE_DEPRECATION_WARNINGS
+	UE_DEPRECATED(all, "The UUsdAssetCache2 class has been deprecated. Use UUsdAssetCache3 and call SetUsdAssetCache instead")
+	UFUNCTION()
 	USDSTAGE_API void SetAssetCache(UUsdAssetCache2* NewCache);
+	PRAGMA_ENABLE_DEPRECATION_WARNINGS
 
 	UFUNCTION(BlueprintCallable, Category = "USD", meta = (CallInEditor = "true"))
 	USDSTAGE_API void SetUsdAssetCache(UUsdAssetCache3* NewCache);
@@ -230,12 +230,12 @@ public:
 	UFUNCTION(BlueprintCallable, Category = "USD", meta = (CallInEditor = "true"))
 	USDSTAGE_API void SetShareAssetsForIdenticalPrims(bool bShare);
 
-	UE_DEPRECATED(5.5, "Use 'SetShareAssetsForIdenticalPrims' instead.")
-	UFUNCTION(BlueprintCallable, Category = "USD", meta = (CallInEditor = "true"))
+	UE_DEPRECATED(all, "Use 'SetShareAssetsForIdenticalPrims' instead.")
+	UFUNCTION()
 	USDSTAGE_API void SetReuseIdenticalAssets(bool bReuse);
 
-	UE_DEPRECATED(5.2, "This option is now controlled via the cvar 'USD.CollapseTopLevelPointInstancers'")
-	UFUNCTION(BlueprintCallable, Category = "USD", meta = (CallInEditor = "true"))
+	UE_DEPRECATED(all, "This option is now controlled via the cvar 'USD.CollapseTopLevelPointInstancers'")
+	UFUNCTION()
 	USDSTAGE_API void SetCollapseTopLevelPointInstancers(bool bCollapse);
 
 	UFUNCTION(BlueprintCallable, Category = "USD", meta = (CallInEditor = "true"))
@@ -332,9 +332,6 @@ public:
 	USDSTAGE_API FString GetIsolatedRootLayer() const;
 
 public:
-	UE_DEPRECATED(5.4, "This function has been renamed into 'GetOrOpenUsdStage', which better describes what it does")
-	USDSTAGE_API UE::FUsdStage& GetOrLoadUsdStage();
-
 	// Opens the stage with RootLayer if its not loaded already, and returns the either the isolated stage (if any) or
 	// the base stage
 	USDSTAGE_API UE::FUsdStage& GetOrOpenUsdStage();
@@ -356,9 +353,6 @@ public:
 	 */
 	USDSTAGE_API void IsolateLayer(const UE::FSdfLayer& Layer, bool bLoadUsdStage = true);
 
-	UE_DEPRECATED(5.5, "This function has been renamed into 'RegenerateLevelSequence', which better describes what it does")
-	USDSTAGE_API void ReloadAnimations();
-
 	// Creates new ULevelSequence assets for the root layer and sublayers.
 	// If we already had a previous LevelSequence opened, we will close its Sequencer and return true
 	USDSTAGE_API bool RegenerateLevelSequence();
@@ -373,10 +367,6 @@ public:
 
 	USDSTAGE_API void OpenLevelSequence();
 
-	PRAGMA_DISABLE_DEPRECATION_WARNINGS
-	UE_DEPRECATED(5.5, "The FUsdInfoCache struct is deprecated in favor of the UUsdInfoCache and UUsdPrimLinkCache, referenced via UPROPERTYs")
-	USDSTAGE_API TSharedPtr<FUsdInfoCache> GetInfoCache();
-	PRAGMA_ENABLE_DEPRECATION_WARNINGS
 	USDSTAGE_API TSharedPtr<UE::FUsdGeomBBoxCache> GetBBoxCache();
 	USDSTAGE_API TMap<FString, TMap<FString, int32>> GetMaterialToPrimvarToUVIndex();
 	USDSTAGE_API const UsdUtils::FBlendShapeMap& GetBlendShapeMap();
diff --git a/Engine/Plugins/Importers/USDImporter/Source/USDStageImporter/Private/USDStageImportContext.cpp b/Engine/Plugins/Importers/USDImporter/Source/USDStageImporter/Private/USDStageImportContext.cpp
--- a/Engine/Plugins/Importers/USDImporter/Source/USDStageImporter/Private/USDStageImportContext.cpp
+++ b/Engine/Plugins/Importers/USDImporter/Source/USDStageImporter/Private/USDStageImportContext.cpp
@@ -108,7 +108,7 @@ void FUsdStageImportContext::Reset()
 	LevelSequenceHelper.Clear();
 	UsdAssetCache = nullptr;
 	PRAGMA_DISABLE_DEPRECATION_WARNINGS
-	AssetCache = nullptr;
+	AssetCache_DEPRECATED = nullptr;
 	PRAGMA_ENABLE_DEPRECATION_WARNINGS
 	BBoxCache = nullptr;
 	MaterialToPrimvarToUVIndex.Empty();
diff --git a/Engine/Plugins/Importers/USDImporter/Source/USDStageImporter/Public/USDStageImportContext.h b/Engine/Plugins/Importers/USDImporter/Source/USDStageImporter/Public/USDStageImportContext.h
--- a/Engine/Plugins/Importers/USDImporter/Source/USDStageImporter/Public/USDStageImportContext.h
+++ b/Engine/Plugins/Importers/USDImporter/Source/USDStageImporter/Public/USDStageImportContext.h
@@ -64,15 +64,12 @@ struct FUsdStageImportContext
 	UPROPERTY()
 	TObjectPtr<UUsdAssetCache3> UsdAssetCache;
 
-	UE_DEPRECATED(5.5, "Use the 'UsdAssetCache' member instead, which is of the new UUsdAssetCache3 type")
+	UE_DEPRECATED(all, "Use the 'UsdAssetCache' member instead, which is of the new UUsdAssetCache3 type")
 	UPROPERTY()
-	TObjectPtr<UUsdAssetCache2> AssetCache;
+	TObjectPtr<UUsdAssetCache2> AssetCache_DEPRECATED;
 
-	PRAGMA_DISABLE_DEPRECATION_WARNINGS
-	/** Caches various information about prims that are expensive to query */
-	UE_DEPRECATED(5.3, "The import process now always builds its own InfoCache, so this member is no longer used")
-	TSharedPtr<FUsdInfoCache> InfoCache;
-	PRAGMA_ENABLE_DEPRECATION_WARNINGS
+	UE_DEPRECATED(all, "The import process now always builds its own InfoCache, so this member is no longer used")
+	TSharedPtr<FUsdInfoCache> InfoCache_DEPRECATED;
 
 	/** Bounding box cache used for the USD stage in case we have to spawn bounds components */
 	TSharedPtr<UE::FUsdGeomBBoxCache> BBoxCache;
diff --git a/Engine/Plugins/Importers/USDImporter/Source/USDStageImporter/Public/USDStageImportOptions.h b/Engine/Plugins/Importers/USDImporter/Source/USDStageImporter/Public/USDStageImportOptions.h
--- a/Engine/Plugins/Importers/USDImporter/Source/USDStageImporter/Public/USDStageImportOptions.h
+++ b/Engine/Plugins/Importers/USDImporter/Source/USDStageImporter/Public/USDStageImportOptions.h
@@ -190,9 +190,9 @@ public:
 	UPROPERTY(BlueprintReadWrite, config, EditAnywhere, Category = "Conflicts with Existing Objects")
 	EReplaceAssetPolicy ExistingAssetPolicy;
 
-	UE_DEPRECATED(5.5, "This property has been renamed to 'Share Assets for Identical Prims'")
+	UE_DEPRECATED(all, "Use 'Share Assets for Identical Prims' instead")
 	UPROPERTY()
-	bool bReuseIdenticalAssets;
+	bool bReuseIdenticalAssets_DEPRECATED;
 
 	/**
 	 * When enabled, assets will be imported into a content folder structure according to their prim path. When disabled,
@@ -228,14 +228,9 @@ public:
 	UPROPERTY(BlueprintReadWrite, config, EditAnywhere, Category = "Processing")
 	bool bMergeIdenticalMaterialSlots;
 
-	/**
-	 * If true, will cause us to collapse any point instancer prim into a single static mesh and static mesh component.
-	 * If false, will cause us to use HierarchicalInstancedStaticMeshComponents to replicate the instancing behavior.
-	 * Point instancers inside other point instancer prototypes are *always* collapsed into the prototype's static mesh.
-	 */
-	UE_DEPRECATED(5.2, "This option is now controlled via the cvar 'USD.CollapseTopLevelPointInstancers'")
+	UE_DEPRECATED(all, "This option is now controlled via the cvar 'USD.CollapseTopLevelPointInstancers'")
 	UPROPERTY()
-	bool bCollapseTopLevelPointInstancers;
+	bool bCollapseTopLevelPointInstancers_DEPRECATED;
 
 	/** When true, if a prim has a "LOD" variant set with variants named "LOD0", "LOD1", etc. where each contains a UsdGeomMesh, the importer will
 	 * attempt to parse the meshes as separate LODs of a single UStaticMesh. When false, only the selected variant will be parsed as LOD0 of the
diff --git a/Engine/Plugins/Importers/USDImporter/Source/USDTests/Private/USDTestsBlueprintLibrary.cpp b/Engine/Plugins/Importers/USDImporter/Source/USDTests/Private/USDTestsBlueprintLibrary.cpp
--- a/Engine/Plugins/Importers/USDImporter/Source/USDTests/Private/USDTestsBlueprintLibrary.cpp
+++ b/Engine/Plugins/Importers/USDImporter/Source/USDTests/Private/USDTestsBlueprintLibrary.cpp
@@ -2,155 +2,31 @@
 
 #include "USDTestsBlueprintLibrary.h"
 
-#include "Objects/USDInfoCache.h"
-#include "USDErrorUtils.h"
-#include "USDStageActor.h"
-
-#include "CoreMinimal.h"
-#include "Editor.h"
-#include "Editor/Transactor.h"
-#include "Kismet2/BlueprintEditorUtils.h"
-#include "Kismet2/KismetEditorUtilities.h"
-
 #include UE_INLINE_GENERATED_CPP_BY_NAME(USDTestsBlueprintLibrary)
 
 bool USDTestsBlueprintLibrary::RecompileBlueprintStageActor(AUsdStageActor* BlueprintDerivedStageActor)
 {
-#if WITH_EDITOR
-	if (!BlueprintDerivedStageActor)
-	{
-		return false;
-	}
-
-	UBlueprintGeneratedClass* BPClass = Cast<UBlueprintGeneratedClass>(BlueprintDerivedStageActor->GetClass());
-	if (!BPClass)
-	{
-		return false;
-	}
-
-	UBlueprint* BP = Cast<UBlueprint>(BPClass->ClassGeneratedBy);
-	if (!BP)
-	{
-		return false;
-	}
-
-	// Compile blueprint. Copied from BlueprintEditorTests.cpp
-	FBlueprintEditorUtils::RefreshAllNodes(BP);
-	FKismetEditorUtilities::CompileBlueprint(BP, EBlueprintCompileOptions::SkipGarbageCollection);
-	if (BP->Status == EBlueprintStatus::BS_UpToDate)
-	{
-		USD_LOG_INFO(TEXT("Blueprint compiled successfully (%s)"), *BP->GetName());
-		return true;
-	}
-	else if (BP->Status == EBlueprintStatus::BS_UpToDateWithWarnings)
-	{
-		USD_LOG_WARNING(TEXT("Blueprint compiled successfully with warnings(%s)"), *BP->GetName());
-		return true;
-	}
-	else if (BP->Status == EBlueprintStatus::BS_Error)
-	{
-		USD_LOG_ERROR(TEXT("Blueprint failed to compile (%s)"), *BP->GetName());
-		return false;
-	}
-
-	USD_LOG_ERROR(TEXT("Blueprint is in an unexpected state after compiling (%s)"), *BP->GetName());
-#endif	  // WITH_EDITOR
 	return false;
 }
 
 void USDTestsBlueprintLibrary::DirtyStageActorBlueprint(AUsdStageActor* BlueprintDerivedStageActor)
 {
-#if WITH_EDITOR
-	if (!BlueprintDerivedStageActor)
-	{
-		return;
-	}
-
-	UBlueprintGeneratedClass* BPClass = Cast<UBlueprintGeneratedClass>(BlueprintDerivedStageActor->GetClass());
-	if (!BPClass)
-	{
-		return;
-	}
-
-	UBlueprint* BP = Cast<UBlueprint>(BPClass->ClassGeneratedBy);
-	if (!BP)
-	{
-		return;
-	}
-
-	// We need to add something that potentially modifies the blueprint code. Just dirtying the blueprint or changing it status doesn't trigger
-	// a recompile when going into PIE
-	FName VarName = FBlueprintEditorUtils::FindUniqueKismetName(BP, TEXT("NewVar"));
-	FEdGraphPinType StringPinType(UEdGraphSchema_K2::PC_String, NAME_None, nullptr, EPinContainerType::None, false, FEdGraphTerminalType());
-	bool bSuccess = FBlueprintEditorUtils::AddMemberVariable(BP, VarName, StringPinType);
-	if (!bSuccess)
-	{
-		USD_LOG_ERROR(TEXT("Failed to add new variable to blueprint (%s)"), *BP->GetName());
-	}
-#endif	  // WITH_EDITOR
 }
 
 int64 USDTestsBlueprintLibrary::GetSubtreeVertexCount(AUsdStageActor* StageActor, const FString& PrimPath)
 {
-	if (StageActor)
-	{
-		if (UUsdInfoCache* Cache = StageActor->UsdInfoCache)
-		{
-			TOptional<uint64> Result = Cache->GetInner().GetSubtreeVertexCount(UE::FSdfPath{*PrimPath});
-			if (Result.IsSet())
-			{
-				// Narrowing conversion here, but we're only using this for our test scenes, which have at most a few thousand vertices
-				return static_cast<int64>(Result.GetValue());
-			}
-			else
-			{
-				return -1;
-			}
-		}
-	}
-
 	return -1;
 }
 
 int64 USDTestsBlueprintLibrary::GetSubtreeMaterialSlotCount(AUsdStageActor* StageActor, const FString& PrimPath)
 {
-	if (StageActor)
-	{
-		if (UUsdInfoCache* Cache = StageActor->UsdInfoCache)
-		{
-			TOptional<uint64> Result = Cache->GetInner().GetSubtreeMaterialSlotCount(UE::FSdfPath{*PrimPath});
-			if (Result.IsSet())
-			{
-				return static_cast<int64>(Result.GetValue());
-			}
-			else
-			{
-				return -1;
-			}
-		}
-	}
-
 	return -1;
 }
 
 void USDTestsBlueprintLibrary::SetUsdStageCpp(AUsdStageActor* StageActor, const FString& NewStageRootLayer)
 {
-	if (!StageActor)
-	{
-		return;
-	}
-
-	UE::FUsdStage NewStage = UnrealUSDWrapper::OpenStage(*NewStageRootLayer, EUsdInitialLoadSet::LoadAll);
-	StageActor->SetUsdStage(NewStage);
 }
 
 void USDTestsBlueprintLibrary::ClearTransactionHistory()
 {
-	if (GEditor)
-	{
-		if (UTransactor* Transactor = GEditor->Trans)
-		{
-			Transactor->Reset(NSLOCTEXT("USDTests", "ClearTransactionHistoryReason", "USDTestsBlueprintLibrary::ClearTransactionHistory was called"));
-		}
-	}
 }
diff --git a/Engine/Plugins/Importers/USDImporter/Source/USDTests/Public/USDTestsBlueprintLibrary.h b/Engine/Plugins/Importers/USDImporter/Source/USDTests/Public/USDTestsBlueprintLibrary.h
--- a/Engine/Plugins/Importers/USDImporter/Source/USDTests/Public/USDTestsBlueprintLibrary.h
+++ b/Engine/Plugins/Importers/USDImporter/Source/USDTests/Public/USDTestsBlueprintLibrary.h
@@ -9,61 +9,35 @@
 
 #define UE_API USDTESTS_API
 
-/**
- * Library of functions that can be used via Python scripting to help testing the other USD functionality
- */
 UCLASS(MinimalAPI, meta = (ScriptName = "USDTestingLibrary"))
-class USDTestsBlueprintLibrary : public UBlueprintFunctionLibrary
+class UE_DEPRECATED(all, "These functions were meant for internal testing only and will be removed in a future release") USDTestsBlueprintLibrary : public UBlueprintFunctionLibrary
 {
 	GENERATED_BODY()
 
 public:
 
-	/**
-	 * Helps test the effects of blueprint recompilation on the spawned actors and assets when a stage is opened.
-	 * Returns whether it compiled successfully or not.
-	 */
-	UE_DEPRECATED(5.4, "These functions were meant for internal testing only and will be removed in a future release.")
-	UFUNCTION(BlueprintCallable, Category = "Blueprint")
+	UE_DEPRECATED(all, "These functions were meant for internal testing only and will be removed in a future release.")
+	UFUNCTION()
 	static UE_API bool RecompileBlueprintStageActor(AUsdStageActor* BlueprintDerivedStageActor);
 
-	/**
-	 * Intentionally dirties the UBlueprint for the given stage actor's generated class.
-	 * This is useful for testing how the stage actor behaves when going into PIE with a dirty blueprint, as that usually triggers
-	 * a recompile at the very sensitive PIE transition
-	 */
-	UE_DEPRECATED(5.4, "These functions were meant for internal testing only and will be removed in a future release.")
-	UFUNCTION(BlueprintCallable, Category = "Blueprint")
+	UE_DEPRECATED(all, "These functions were meant for internal testing only and will be removed in a future release.")
+	UFUNCTION()
 	static UE_API void DirtyStageActorBlueprint(AUsdStageActor* BlueprintDerivedStageActor);
 
-	/**
-	 * Queries a subtree vertex count using the stage actor's info cache, which is not yet exposed to blueprint.
-	 * May return -1 in case of an error.
-	 */
-	UE_DEPRECATED(5.4, "These functions were meant for internal testing only and will be removed in a future release.")
-	UFUNCTION(BlueprintCallable, Category = "Counts")
+	UE_DEPRECATED(all, "These functions were meant for internal testing only and will be removed in a future release.")
+	UFUNCTION()
 	static UE_API int64 GetSubtreeVertexCount(AUsdStageActor* StageActor, const FString& PrimPath);
 
-	/**
-	 * Queries a subtree material slot count using the stage actor's info cache, which is not yet exposed to blueprint.
-	 * May return -1 in case of an error.
-	 */
-	UE_DEPRECATED(5.4, "These functions were meant for internal testing only and will be removed in a future release.")
-	UFUNCTION(BlueprintCallable, Category = "Counts")
+	UE_DEPRECATED(all, "These functions were meant for internal testing only and will be removed in a future release.")
+	UFUNCTION()
 	static UE_API int64 GetSubtreeMaterialSlotCount(AUsdStageActor* StageActor, const FString& PrimPath);
 
-	/**
-	 * Internally opens the provided stage and sets it on the stage actor via C++
-	 */
-	UE_DEPRECATED(5.4, "These functions were meant for internal testing only and will be removed in a future release.")
-	UFUNCTION(BlueprintCallable, Category = "Counts")
+	UE_DEPRECATED(all, "These functions were meant for internal testing only and will be removed in a future release.")
+	UFUNCTION()
 	static UE_API void SetUsdStageCpp(AUsdStageActor* StageActor, const FString& NewStageRootLayer);
 
-	/**
-	 * Clears the editor transaction history (useful to be run before GC, to check for otherwise unreferenced assets)
-	 */
-	UE_DEPRECATED(5.4, "These functions were meant for internal testing only and will be removed in a future release.")
-	UFUNCTION(BlueprintCallable, Category = "Transactions")
+	UE_DEPRECATED(all, "These functions were meant for internal testing only and will be removed in a future release.")
+	UFUNCTION()
 	static UE_API void ClearTransactionHistory();
 };
 
diff --git a/Engine/Plugins/Importers/USDImporterMDL/Source/USDImporterMDL/Private/MDLUSDShadeMaterialTranslator.cpp b/Engine/Plugins/Importers/USDImporterMDL/Source/USDImporterMDL/Private/MDLUSDShadeMaterialTranslator.cpp
--- a/Engine/Plugins/Importers/USDImporterMDL/Source/USDImporterMDL/Private/MDLUSDShadeMaterialTranslator.cpp
+++ b/Engine/Plugins/Importers/USDImporterMDL/Source/USDImporterMDL/Private/MDLUSDShadeMaterialTranslator.cpp
@@ -35,10 +35,6 @@
 
 #define LOCTEXT_NAMESPACE "MDLUSDShadeMaterialTranslator"
 
-PRAGMA_DISABLE_DEPRECATION_WARNINGS
-FName FMdlUsdShadeMaterialTranslator::MdlRenderContext = TEXT("mdl");
-PRAGMA_ENABLE_DEPRECATION_WARNINGS
-
 namespace UE::MDLShadeTranslator::Private
 {
 	void NotifyIfMaterialNeedsVirtualTextures(UMaterialInterface* MaterialInterface)
diff --git a/Engine/Plugins/Importers/USDImporterMDL/Source/USDImporterMDL/Public/MDLUSDShadeMaterialTranslator.h b/Engine/Plugins/Importers/USDImporterMDL/Source/USDImporterMDL/Public/MDLUSDShadeMaterialTranslator.h
--- a/Engine/Plugins/Importers/USDImporterMDL/Source/USDImporterMDL/Public/MDLUSDShadeMaterialTranslator.h
+++ b/Engine/Plugins/Importers/USDImporterMDL/Source/USDImporterMDL/Public/MDLUSDShadeMaterialTranslator.h
@@ -16,10 +16,6 @@ class FMdlUsdShadeMaterialTranslator : public FMaterialXUsdShadeMaterialTranslat
 {
 	using Super = FMaterialXUsdShadeMaterialTranslator;
 
-public:
-	UE_DEPRECATED(5.5, "Use UnrealIdentifiers::MdlRenderContext instead")
-	static FName MdlRenderContext;
-
 public:
 	using FMaterialXUsdShadeMaterialTranslator::FMaterialXUsdShadeMaterialTranslator;
 
diff --git a/Engine/Plugins/Runtime/USDCore/Source/USDClasses/Private/USDAssetCache2.cpp b/Engine/Plugins/Runtime/USDCore/Source/USDClasses/Private/USDAssetCache2.cpp
--- a/Engine/Plugins/Runtime/USDCore/Source/USDClasses/Private/USDAssetCache2.cpp
+++ b/Engine/Plugins/Runtime/USDCore/Source/USDClasses/Private/USDAssetCache2.cpp
@@ -25,6 +25,8 @@
 
 #include UE_INLINE_GENERATED_CPP_BY_NAME(USDAssetCache2)
 
+PRAGMA_DISABLE_DEPRECATION_WARNINGS
+
 static bool GOnDemandCachedAssetLoading = true;
 static FAutoConsoleVariableRef COnDemandCachedAssetLoading(
 	TEXT("USD.OnDemandCachedAssetLoading"),
@@ -1849,3 +1851,5 @@ FUsdScopedAssetCacheReferencer::~FUsdScopedAssetCacheReferencer()
 		ValidCache->CurrentScopedReferencer = nullptr;
 	}
 }
+
+PRAGMA_ENABLE_DEPRECATION_WARNINGS
diff --git a/Engine/Plugins/Runtime/USDCore/Source/USDClasses/Private/USDClassesModule.cpp b/Engine/Plugins/Runtime/USDCore/Source/USDClasses/Private/USDClassesModule.cpp
--- a/Engine/Plugins/Runtime/USDCore/Source/USDClasses/Private/USDClassesModule.cpp
+++ b/Engine/Plugins/Runtime/USDCore/Source/USDClasses/Private/USDClassesModule.cpp
@@ -785,26 +785,6 @@ UUsdAssetCache3* IUsdClassesModule::GetAssetCacheForProject()
 	return DefaultCache;
 }
 
-FString IUsdClassesModule::SanitizeObjectName(const FString& InObjectName)
-{
-	return UsdUnreal::ObjectUtils::SanitizeObjectName(InObjectName);
-}
-
-const FSoftObjectPath* IUsdClassesModule::GetReferenceMaterialPath(const FDisplayColorMaterial& DisplayColorDescription)
-{
-	return UsdUnreal::MaterialUtils::GetReferenceMaterialPath(DisplayColorDescription);
-}
-
-UMaterialInstanceDynamic* IUsdClassesModule::CreateDisplayColorMaterialInstanceDynamic(const FDisplayColorMaterial& DisplayColorDescription)
-{
-	return UsdUnreal::MaterialUtils::CreateDisplayColorMaterialInstanceDynamic(DisplayColorDescription);
-}
-
-UMaterialInstanceConstant* IUsdClassesModule::CreateDisplayColorMaterialInstanceConstant(const FDisplayColorMaterial& DisplayColorDescription)
-{
-	return UsdUnreal::MaterialUtils::CreateDisplayColorMaterialInstanceConstant(DisplayColorDescription);
-}
-
 class FUsdClassesModule : public IUsdClassesModule
 {
 public:
diff --git a/Engine/Plugins/Runtime/USDCore/Source/USDClasses/Public/USDAssetCache.h b/Engine/Plugins/Runtime/USDCore/Source/USDClasses/Public/USDAssetCache.h
--- a/Engine/Plugins/Runtime/USDCore/Source/USDClasses/Public/USDAssetCache.h
+++ b/Engine/Plugins/Runtime/USDCore/Source/USDClasses/Public/USDAssetCache.h
@@ -8,12 +8,8 @@
 
 #define UE_API USDCLASSES_API
 
-class UE_DEPRECATED(5.2, "Use the UUsdAssetCache2 class instead, which is an actual asset you can create on the Content Browser") UUsdAssetCache;
-
-/** Owns the assets generated and reused by the USD stage, allowing thread-safe retrieval/storage */
 UCLASS(meta = (ScriptName = "UsdAssetCache_Deprecated"))
-
-class UUsdAssetCache : public UObject
+class UE_DEPRECATED(all, "Use the UUsdAssetCache2 class instead, which is an actual asset you can create on the Content Browser") UUsdAssetCache : public UObject
 {
 	GENERATED_BODY()
 
diff --git a/Engine/Plugins/Runtime/USDCore/Source/USDClasses/Public/USDAssetCache2.h b/Engine/Plugins/Runtime/USDCore/Source/USDClasses/Public/USDAssetCache2.h
--- a/Engine/Plugins/Runtime/USDCore/Source/USDClasses/Public/USDAssetCache2.h
+++ b/Engine/Plugins/Runtime/USDCore/Source/USDClasses/Public/USDAssetCache2.h
@@ -11,11 +11,8 @@
 
 #define UE_API USDCLASSES_API
 
-/**
- * Owns the assets generated and reused by USD Stages, allowing thread-safe retrieval/storage.
- */
 UCLASS(MinimalAPI, BlueprintType, Blueprintable)
-class UUsdAssetCache2 : public UObject
+class UE_DEPRECATED(5.8, "Use the UUsdAssetCache3 class instead, which exposes the cached assets on the Content Browser") UUsdAssetCache2 : public UObject
 {
 	GENERATED_BODY()
 
@@ -252,18 +249,8 @@ private:
 	const UObject* CurrentScopedReferencer = nullptr;
 };	  // UCLASS(BlueprintType)
 
-/**
- * The UUsdAssetCache2 can track all the UObjects that are referencing assets, so that it knows when to discard an
- * unreferenced asset.
- *
- * This struct lets you specify a single UObject that will be automatically used as the referencer for the given asset
- * cache for the duration of the scope.
- *
- * Usage:
- *     FUsdScopedAssetCacheReferencer ScopedReferencer{ StageActor->UsdAssetCache, StageActor };
- *     StageActor->SetRootLayer(MyRootLayer);  // Adds assets to the cache
- */
-struct FUsdScopedAssetCacheReferencer
+PRAGMA_DISABLE_DEPRECATION_WARNINGS
+struct UE_DEPRECATED(5.8, "No longer used as it's not needed by the UsdAssetCache3") FUsdScopedAssetCacheReferencer
 {
 public:
 	UE_API explicit FUsdScopedAssetCacheReferencer(UUsdAssetCache2* InAssetCache, const UObject* Referencer);
@@ -278,5 +265,6 @@ public:
 private:
 	TWeakObjectPtr<UUsdAssetCache2> AssetCache;
 };
+PRAGMA_ENABLE_DEPRECATION_WARNINGS
 
 #undef UE_API
diff --git a/Engine/Plugins/Runtime/USDCore/Source/USDClasses/Public/USDAssetImportData.h b/Engine/Plugins/Runtime/USDCore/Source/USDClasses/Public/USDAssetImportData.h
--- a/Engine/Plugins/Runtime/USDCore/Source/USDClasses/Public/USDAssetImportData.h
+++ b/Engine/Plugins/Runtime/USDCore/Source/USDClasses/Public/USDAssetImportData.h
@@ -14,9 +14,9 @@ class UUsdAssetImportData : public UAssetImportData
 	GENERATED_BODY()
 
 public:
-	UE_DEPRECATED(5.3, "Use USD AssetUserData instead of AssetImportData for USD-specific info")
+	UE_DEPRECATED(all, "Use USD AssetUserData instead of AssetImportData for USD-specific info")
 	UPROPERTY()
-	FString PrimPath;
+	FString PrimPath_DEPRECATED;
 
 	// Likely a UUSDStageImportOptions, but we don't declare it here
 	// to prevent an unnecessary module dependency on USDStageImporter
@@ -26,41 +26,23 @@ public:
 
 PRAGMA_DISABLE_DEPRECATION_WARNINGS
 UCLASS()
-class UE_DEPRECATED(5.3, "Use USD AssetUserData instead of AssetImportData for USD-specific info") USDCLASSES_API UUsdAnimSequenceAssetImportData
+class UE_DEPRECATED(all, "Use USD AssetUserData instead of AssetImportData for USD-specific info") USDCLASSES_API UUsdAnimSequenceAssetImportData
 	: public UUsdAssetImportData
 {
 	GENERATED_BODY()
 
 public:
-	// Offset into the stage in seconds to when this AnimSequence should start
-	// playing to match the skeletal animation in that stage. This is required
-	// because UAnimSequences just range from the first skeletal timeSample to the
-	// last, and the first sample is not necessarily the startTimeCode for the
-	// stage. Note that this is wrt. the startTimeCode of the *layer*, and not the
-	// composed stage. This is only used for animating USkeletalMeshComponents via
-	// the AUsdStageActor, and it is in seconds since it will need to drive the
-	// UAnimSequence with a time value in the [0, LenghtSeconds] range. This
-	// should be applied *after* any offset/scale conversions on the time
-	// coordinate.
 	UPROPERTY()
 	float LayerStartOffsetSeconds = 0.0f;
 };
 
-/** We assign these to UStaticMeshes or USkeletalMeshes generated from USD */
 UCLASS()
-class UE_DEPRECATED(5.3, "Use USD AssetUserData instead of AssetImportData for USD-specific info") USDCLASSES_API UUsdMeshAssetImportData
+class UE_DEPRECATED(all, "Use USD AssetUserData instead of AssetImportData for USD-specific info") USDCLASSES_API UUsdMeshAssetImportData
 	: public UUsdAssetImportData
 {
 	GENERATED_BODY()
 
 public:
-	/**
-	 * Maps from a material slot index of an UStaticMesh or USkeletalMesh to a
-	 * list of source prims that contain this assignment. It can contain multiple
-	 * prims in case we combine material slots and/or collapse prims (e.g. {0:
-	 * ['/Root/mesh', '/Root/othermesh/geomsubset0',
-	 * '/Root/othermesh/geomsubset1'] }).
-	 */
 	UPROPERTY()
 	TMap<int32, FUsdPrimPathList> MaterialSlotToPrimPaths;
 };
diff --git a/Engine/Plugins/Runtime/USDCore/Source/USDClasses/Public/USDAssetUserData.h b/Engine/Plugins/Runtime/USDCore/Source/USDClasses/Public/USDAssetUserData.h
--- a/Engine/Plugins/Runtime/USDCore/Source/USDClasses/Public/USDAssetUserData.h
+++ b/Engine/Plugins/Runtime/USDCore/Source/USDClasses/Public/USDAssetUserData.h
@@ -44,7 +44,7 @@ public:
 	// individual layer. This is used for animating USkeletalMeshComponents via
 	// the AUsdStageActor and building the LevelSequences, and it is in seconds 
 	// since it will need to drive the UAnimSequence with a time value in the 
-	// [0, LenghtSeconds] range. This should be applied *after* any offset/scale 
+	// [0, LengthSeconds] range. This should be applied *after* any offset/scale 
 	// conversions on the time coordinate.
 	UPROPERTY()
 	float StageStartOffsetSeconds = 0.0f;
@@ -136,28 +136,16 @@ public:
 	float LayerStartOffsetSeconds = 0.0f;
 };
 
-/**
- * We assign these to persistent LevelSequences that bind to one of the actors/components that the stage actor spawns.
- * We need this as part of a mechanism to automatically repair those bindings when they break if we close/reload the stage.
- */
 UCLASS(BlueprintType)
-class UE_DEPRECATED(5.5, "This class is no longer used, as AUsdStageActors now use Sequencer Universal Object Locators for persistent bindings")
+class UE_DEPRECATED(all, "This class is no longer used, as AUsdStageActors now use Sequencer Universal Object Locators for persistent bindings")
 	USDCLASSES_API UUsdLevelSequenceAssetUserData : public UAssetUserData
 {
 	GENERATED_BODY()
 
 public:
-	/**
-	 * The LevelSequence has a Guid that is changed every time its state is modified.
-	 * We pay attention to that so that we can avoid reprocessing a LevelSequence that hasn't changed
-	 */
 	UPROPERTY(EditAnywhere, BlueprintReadWrite, Category = "USD")
 	FGuid LastCheckedSignature;
 
-	/**
-	 * Set of binding GUIDs that we already handled in the past.
-	 * We use this so that we won't try and overwrite the changes in case the user manually clears/modifies a binding we previously setup.
-	 */
 	UPROPERTY(EditAnywhere, BlueprintReadWrite, Category = "USD")
 	TSet<FGuid> HandledBindingGuids;
 };
diff --git a/Engine/Plugins/Runtime/USDCore/Source/USDClasses/Public/USDClassesModule.h b/Engine/Plugins/Runtime/USDCore/Source/USDClasses/Public/USDClassesModule.h
--- a/Engine/Plugins/Runtime/USDCore/Source/USDClasses/Public/USDClassesModule.h
+++ b/Engine/Plugins/Runtime/USDCore/Source/USDClasses/Public/USDClassesModule.h
@@ -63,20 +63,4 @@ public:
 
 	// Returns the default asset cache for the project or create a new one at the project root
 	USDCLASSES_API static UUsdAssetCache3* GetAssetCacheForProject();
-
-	// Adapted from ObjectTools as it is within an Editor-only module
-	UE_DEPRECATED(5.5, "This function has been moved to USDObjectUtils.h")
-	USDCLASSES_API static FString SanitizeObjectName(const FString& InObjectName);
-
-	// Backwards compatibility for 5.5 as we moved this struct into MaterialUtils
-	using FDisplayColorMaterial = UsdUnreal::MaterialUtils::FDisplayColorMaterial;
-
-	UE_DEPRECATED(5.5, "This function has been moved to USDMaterialUtils.h")
-	USDCLASSES_API static const FSoftObjectPath* GetReferenceMaterialPath(const FDisplayColorMaterial& DisplayColorDescription);
-
-	UE_DEPRECATED(5.5, "This function has been moved to USDMaterialUtils.h")
-	USDCLASSES_API static UMaterialInstanceDynamic* CreateDisplayColorMaterialInstanceDynamic(const FDisplayColorMaterial& DisplayColorDescription);
-
-	UE_DEPRECATED(5.5, "This function has been moved to USDMaterialUtils.h")
-	USDCLASSES_API static UMaterialInstanceConstant* CreateDisplayColorMaterialInstanceConstant(const FDisplayColorMaterial& DisplayColorDescription);
 };
diff --git a/Engine/Plugins/Runtime/USDCore/Source/USDUtilities/Private/Objects/USDInfoCache.cpp b/Engine/Plugins/Runtime/USDCore/Source/USDUtilities/Private/Objects/USDInfoCache.cpp
--- a/Engine/Plugins/Runtime/USDCore/Source/USDUtilities/Private/Objects/USDInfoCache.cpp
+++ b/Engine/Plugins/Runtime/USDCore/Source/USDUtilities/Private/Objects/USDInfoCache.cpp
@@ -595,23 +595,6 @@ TArray<UE::FSdfPath> FUsdInfoCache::GetChildren(const UE::FSdfPath& ParentPath)
 	return Result;
 }
 
-PRAGMA_DISABLE_DEPRECATION_WARNINGS
-TSet<UE::FSdfPath> FUsdInfoCache::GetKnownPrims() const
-{
-	TSet<UE::FSdfPath> Result;
-
-	if (FUsdInfoCacheImpl* ImplPtr = Impl.Get())
-	{
-		FReadScopeLock ScopeLock(ImplPtr->InfoMapLock);
-
-		ImplPtr->InfoMap.GetKeys(Result);
-		return Result;
-	}
-
-	return Result;
-}
-PRAGMA_ENABLE_DEPRECATION_WARNINGS
-
 bool FUsdInfoCache::IsPathCollapsed(const UE::FSdfPath& Path, ECollapsingType CollapsingType) const
 {
 	using namespace UE::UsdInfoCache::Private;
@@ -2288,28 +2271,6 @@ namespace UE::USDInfoCache::Private
 #endif	  // USE_SD_SDK
 }	 // namespace UE::USDInfoCache::Private
 
-PRAGMA_DISABLE_DEPRECATION_WARNINGS
-bool FUsdInfoCache::IsPotentialGeometryCacheRoot(const UE::FSdfPath& Path) const
-{
-	if (FUsdInfoCacheImpl* ImplPtr = Impl.Get())
-	{
-		FReadScopeLock ScopeLock(ImplPtr->InfoMapLock);
-
-		if (const UE::UsdInfoCache::Private::FUsdPrimInfo* FoundInfo = ImplPtr->GetPrimInfo(Path))
-		{
-			FReadScopeLock PrimLock = ImplPtr->LockForReading(*FoundInfo);
-
-			return FoundInfo->GeometryCacheState == UE::UsdInfoCache::Private::EGeometryCachePrimState::ValidRoot;
-		}
-
-		// This should never happen: We should have cached the entire tree
-		ensureMsgf(false, TEXT("Prim path '%s' has not been cached!"), *Path.GetString());
-	}
-
-	return false;
-}
-PRAGMA_ENABLE_DEPRECATION_WARNINGS
-
 bool FUsdInfoCache::IsPotentialGeometryCacheRoot(const UE::FUsdPrim& Prim) const
 {
 	using namespace UE::UsdInfoCache::Private;
@@ -2422,45 +2383,6 @@ TOptional<TArray<UsdUtils::FUsdPrimMaterialSlot>> FUsdInfoCache::GetSubtreeMater
 	return {};
 }
 
-PRAGMA_DISABLE_DEPRECATION_WARNINGS
-void FUsdInfoCache::LinkAssetToPrim(const UE::FSdfPath& Path, UObject* Asset)
-{
-}
-
-void FUsdInfoCache::UnlinkAssetFromPrim(const UE::FSdfPath& Path, UObject* Asset)
-{
-}
-
-TArray<TWeakObjectPtr<UObject>> FUsdInfoCache::RemoveAllAssetPrimLinks(const UE::FSdfPath& Path)
-{
-	return {};
-}
-
-TArray<UE::FSdfPath> FUsdInfoCache::RemoveAllAssetPrimLinks(const UObject* Asset)
-{
-	return {};
-}
-
-void FUsdInfoCache::RemoveAllAssetPrimLinks()
-{
-}
-
-TArray<TWeakObjectPtr<UObject>> FUsdInfoCache::GetAllAssetsForPrim(const UE::FSdfPath& Path) const
-{
-	return {};
-}
-
-TArray<UE::FSdfPath> FUsdInfoCache::GetPrimsForAsset(const UObject* Asset) const
-{
-	return {};
-}
-
-TMap<UE::FSdfPath, TArray<TWeakObjectPtr<UObject>>> FUsdInfoCache::GetAllAssetPrimLinks() const
-{
-	return {};
-}
-PRAGMA_ENABLE_DEPRECATION_WARNINGS
-
 void FUsdInfoCache::RebuildCacheForSubtree(const UE::FUsdPrim& Prim, FUsdSchemaTranslationContext& Context)
 {
 	RebuildCacheForSubtrees({Prim.GetPrimPath()}, Context);
diff --git a/Engine/Plugins/Runtime/USDCore/Source/USDUtilities/Private/Objects/USDSchemaTranslator.cpp b/Engine/Plugins/Runtime/USDCore/Source/USDUtilities/Private/Objects/USDSchemaTranslator.cpp
--- a/Engine/Plugins/Runtime/USDCore/Source/USDUtilities/Private/Objects/USDSchemaTranslator.cpp
+++ b/Engine/Plugins/Runtime/USDCore/Source/USDUtilities/Private/Objects/USDSchemaTranslator.cpp
@@ -62,39 +62,6 @@ TSharedPtr<FUsdSchemaTranslator> FUsdSchemaTranslatorRegistry::CreateTranslatorF
 	return {};
 }
 
-PRAGMA_DISABLE_DEPRECATION_WARNINGS
-FUsdRenderContextRegistry::FUsdRenderContextRegistry()
-{
-}
-
-void FUsdRenderContextRegistry::Register(const FName& RenderContextToken)
-{
-	UsdUnreal::MaterialUtils::RegisterRenderContext(RenderContextToken);
-}
-
-void FUsdRenderContextRegistry::Unregister(const FName& RenderContextToken)
-{
-	UsdUnreal::MaterialUtils::UnregisterRenderContext(RenderContextToken);
-}
-
-const TSet<FName>& FUsdRenderContextRegistry::GetRenderContexts() const
-{
-	static TSet<FName> TempValue;
-	TempValue = TSet<FName>{UsdUnreal::MaterialUtils::GetRegisteredRenderContexts()};
-	return TempValue;
-}
-
-const FName& FUsdRenderContextRegistry::GetUniversalRenderContext() const
-{
-	return UnrealIdentifiers::UniversalRenderContext;
-}
-
-const FName& FUsdRenderContextRegistry::GetUnrealRenderContext() const
-{
-	return UnrealIdentifiers::UnrealRenderContext;
-}
-PRAGMA_ENABLE_DEPRECATION_WARNINGS
-
 FRegisteredSchemaTranslatorHandle FUsdSchemaTranslatorRegistry::Register(const FString& SchemaName, FCreateTranslator CreateFunction)
 {
 #if USE_USD_SDK
@@ -166,13 +133,6 @@ int32 FUsdSchemaTranslatorRegistry::GetExternalSchemaTranslatorCount()
 	return ExternalSchemaTranslatorCount;
 }
 
-FUsdSchemaTranslationContext::FUsdSchemaTranslationContext(const UE::FUsdStage& InStage, UUsdAssetCache2& InAssetCache)
-	: Stage(InStage)
-	, UsdAssetCache(NewObject<UUsdAssetCache3>())
-{
-	RenderContext = UnrealIdentifiers::UniversalRenderContext;
-}
-
 FUsdSchemaTranslationContext::FUsdSchemaTranslationContext(const UE::FUsdStage& InStage)
 	: Stage(InStage)
 {
diff --git a/Engine/Plugins/Runtime/USDCore/Source/USDUtilities/Private/USDConversionUtils.cpp b/Engine/Plugins/Runtime/USDCore/Source/USDUtilities/Private/USDConversionUtils.cpp
--- a/Engine/Plugins/Runtime/USDCore/Source/USDUtilities/Private/USDConversionUtils.cpp
+++ b/Engine/Plugins/Runtime/USDCore/Source/USDUtilities/Private/USDConversionUtils.cpp
@@ -852,22 +852,6 @@ TUsdStore<pxr::TfToken> UsdUtils::GetUVSetName(int32 UVChannelIndex)
 	return UVSetNameToken;
 }
 
-int32 UsdUtils::GetPrimvarUVIndex(FString PrimvarName)
-{
-	int32 Index = PrimvarName.Len();
-	while (Index > 0 && PrimvarName[Index - 1] >= '0' && PrimvarName[Index - 1] <= '9')
-	{
-		--Index;
-	}
-
-	if (Index < PrimvarName.Len())
-	{
-		return FCString::Atoi(*PrimvarName.RightChop(Index));
-	}
-
-	return 0;
-}
-
 TArray<TUsdStore<pxr::UsdGeomPrimvar>> UsdUtils::GetUVSetPrimvars(const pxr::UsdPrim& UsdPrim, int32 MaxNumPrimvars)
 {
 	if (!UsdPrim)
@@ -937,25 +921,6 @@ TArray<TUsdStore<pxr::UsdGeomPrimvar>> UsdUtils::GetUVSetPrimvars(const pxr::Usd
 	return Result;
 }
 
-TArray<TUsdStore<pxr::UsdGeomPrimvar>> UsdUtils::GetUVSetPrimvars(
-	const pxr::UsdGeomMesh& UsdMesh,
-	const TMap<FString, TMap<FString, int32>>& MaterialToPrimvarsUVSetNames,
-	const pxr::TfToken& RenderContext,
-	const pxr::TfToken& MaterialPurpose
-)
-{
-	return UsdUtils::GetUVSetPrimvars(UsdMesh.GetPrim());
-}
-
-TArray<TUsdStore<pxr::UsdGeomPrimvar>> UsdUtils::GetUVSetPrimvars(
-	const pxr::UsdGeomMesh& UsdMesh,
-	const TMap<FString, TMap<FString, int32>>& MaterialToPrimvarsUVSetNames,
-	const UsdUtils::FUsdPrimMaterialAssignmentInfo& UsdMeshMaterialAssignmentInfo
-)
-{
-	return UsdUtils::GetUVSetPrimvars(UsdMesh.GetPrim());
-}
-
 TArray<TUsdStore<pxr::UsdGeomPrimvar>> UsdUtils::AssemblePrimvarsIntoUVSets(
 	const TArray<TUsdStore<pxr::UsdGeomPrimvar>>& AllMeshUVPrimvars,
 	const TMap<FString, int32>& AllowedPrimvarsToUVIndex
@@ -2385,215 +2350,6 @@ double UsdUtils::GetEarliestTimeCode()
 #endif
 }
 
-// We can't just redirect the functions to USDObjectUtils.h because of the module dependencies
-PRAGMA_DISABLE_DEPRECATION_WARNINGS
-UUsdAssetImportData* UsdUtils::GetAssetImportData(UObject* Asset)
-{
-	UUsdAssetImportData* ImportData = nullptr;
-#if WITH_EDITORONLY_DATA
-	if (UStaticMesh* Mesh = Cast<UStaticMesh>(Asset))
-	{
-		ImportData = Cast<UUsdAssetImportData>(Mesh->GetAssetImportData());
-	}
-	else if (USkeleton* Skeleton = Cast<USkeleton>(Asset))
-	{
-		if (USkeletalMesh* SkMesh = Skeleton->GetPreviewMesh())
-		{
-			ImportData = Cast<UUsdAssetImportData>(SkMesh->GetAssetImportData());
-		}
-	}
-	else if (UPhysicsAsset* PhysicsAsset = Cast<UPhysicsAsset>(Asset))
-	{
-		if (USkeletalMesh* SkMesh = PhysicsAsset->GetPreviewMesh())
-		{
-			ImportData = Cast<UUsdAssetImportData>(SkMesh->GetAssetImportData());
-		}
-	}
-	else if (UAnimBlueprint* AnimBP = Cast<UAnimBlueprint>(Asset))
-	{
-		// We will always have a skeleton, but not necessarily we will have a preview mesh directly
-		// on the UAnimBlueprint
-		if (USkeleton* AnimBPSkeleton = AnimBP->TargetSkeleton.Get())
-		{
-			if (USkeletalMesh* SkMesh = AnimBPSkeleton->GetPreviewMesh())
-			{
-				ImportData = Cast<UUsdAssetImportData>(SkMesh->GetAssetImportData());
-			}
-		}
-	}
-	else if (USkeletalMesh* SkMesh = Cast<USkeletalMesh>(Asset))
-	{
-		ImportData = Cast<UUsdAssetImportData>(SkMesh->GetAssetImportData());
-	}
-	else if (UAnimSequence* SkelAnim = Cast<UAnimSequence>(Asset))
-	{
-		ImportData = Cast<UUsdAssetImportData>(SkelAnim->AssetImportData);
-	}
-	else if (UMaterialInterface* Material = Cast<UMaterialInterface>(Asset))
-	{
-		ImportData = Cast<UUsdAssetImportData>(Material->AssetImportData);
-	}
-	else if (UTexture* Texture = Cast<UTexture>(Asset))
-	{
-		ImportData = Cast<UUsdAssetImportData>(Texture->AssetImportData);
-	}
-	else if (UGeometryCache* GeometryCache = Cast<UGeometryCache>(Asset))
-	{
-		ImportData = Cast<UUsdAssetImportData>(GeometryCache->AssetImportData);
-	}
-	else if (UGroomAsset* Groom = Cast<UGroomAsset>(Asset))
-	{
-		ImportData = Cast<UUsdAssetImportData>(Groom->AssetImportData);
-	}
-	else if (UGroomCache* GroomCache = Cast<UGroomCache>(Asset))
-	{
-		ImportData = Cast<UUsdAssetImportData>(GroomCache->AssetImportData);
-	}
-	else if (UStreamableSparseVolumeTexture* SparseVolumeTexture = Cast<UStreamableSparseVolumeTexture>(Asset))
-	{
-		ImportData = Cast<UUsdAssetImportData>(SparseVolumeTexture->AssetImportData);
-	}
-
-#endif
-	return ImportData;
-}
-
-void UsdUtils::SetAssetImportData(UObject* Asset, UAssetImportData* ImportData)
-{
-	if (!Asset)
-	{
-		return;
-	}
-
-#if WITH_EDITOR
-	if (UStaticMesh* Mesh = Cast<UStaticMesh>(Asset))
-	{
-		Mesh->SetAssetImportData(ImportData);
-	}
-	else if (USkeletalMesh* SkMesh = Cast<USkeletalMesh>(Asset))
-	{
-		SkMesh->SetAssetImportData(ImportData);
-	}
-	else if (UAnimSequence* SkelAnim = Cast<UAnimSequence>(Asset))
-	{
-		SkelAnim->AssetImportData = ImportData;
-	}
-	else if (UMaterialInterface* Material = Cast<UMaterialInterface>(Asset))
-	{
-		Material->AssetImportData = ImportData;
-	}
-	else if (UTexture* Texture = Cast<UTexture>(Asset))
-	{
-		Texture->AssetImportData = ImportData;
-	}
-	else if (UGeometryCache* GeometryCache = Cast<UGeometryCache>(Asset))
-	{
-		GeometryCache->AssetImportData = ImportData;
-	}
-	else if (UGroomAsset* Groom = Cast<UGroomAsset>(Asset))
-	{
-		Groom->AssetImportData = ImportData;
-	}
-	else if (UGroomCache* GroomCache = Cast<UGroomCache>(Asset))
-	{
-		GroomCache->AssetImportData = ImportData;
-	}
-	else if (UStreamableSparseVolumeTexture* SparseVolumeTexture = Cast<UStreamableSparseVolumeTexture>(Asset))
-	{
-		SparseVolumeTexture->AssetImportData = ImportData;
-	}
-#endif	  // WITH_EDITOR
-}
-
-UUsdAssetUserData* UsdUtils::GetAssetUserData(const UObject* Object, TSubclassOf<UUsdAssetUserData> Class)
-{
-	if (!Object)
-	{
-		return nullptr;
-	}
-
-	if (!Class)
-	{
-		Class = UUsdAssetUserData::StaticClass();
-	}
-
-	const IInterface_AssetUserData* AssetUserDataInterface = Cast<const IInterface_AssetUserData>(Object);
-	if (!AssetUserDataInterface)
-	{
-		USD_LOG_WARNING(
-			TEXT("Tried getting AssetUserData from object '%s', but the class '%s' doesn't implement the AssetUserData interface!"),
-			*Object->GetPathName(),
-			*Object->GetClass()->GetName()
-		);
-		return nullptr;
-	}
-
-	// Const cast because there is no const access of asset user data on the interface
-	return Cast<UUsdAssetUserData>(const_cast<IInterface_AssetUserData*>(AssetUserDataInterface)->GetAssetUserDataOfClass(Class));
-}
-
-UUsdAssetUserData* UsdUtils::GetOrCreateAssetUserData(UObject* Object, TSubclassOf<UUsdAssetUserData> Class)
-{
-	if (!Object)
-	{
-		return nullptr;
-	}
-
-	if (!Class)
-	{
-		Class = UUsdAssetUserData::StaticClass();
-	}
-
-	IInterface_AssetUserData* AssetUserDataInterface = Cast<IInterface_AssetUserData>(Object);
-	if (!AssetUserDataInterface)
-	{
-		USD_LOG_WARNING(
-			TEXT("Tried adding AssetUserData to object '%s', but it doesn't implement the AssetUserData interface!"),
-			*Object->GetPathName()
-		);
-		return nullptr;
-	}
-
-	UUsdAssetUserData* AssetUserData = Cast<UUsdAssetUserData>(AssetUserDataInterface->GetAssetUserDataOfClass(Class));
-	if (!AssetUserData)
-	{
-		// For now we're expecting objects to only have one instance of UUsdAssetUserData
-		ensure(!AssetUserDataInterface->HasAssetUserDataOfClass(UUsdAssetUserData::StaticClass()));
-
-		AssetUserData = NewObject<UUsdAssetUserData>(Object, Class, TEXT("UsdAssetUserData"), RF_Public);
-		AssetUserDataInterface->AddAssetUserData(AssetUserData);
-	}
-
-	return AssetUserData;
-}
-
-bool UsdUtils::SetAssetUserData(UObject* Object, UUsdAssetUserData* AssetUserData)
-{
-	if (!Object)
-	{
-		return false;
-	}
-
-	IInterface_AssetUserData* AssetUserDataInterface = Cast<IInterface_AssetUserData>(Object);
-	if (!AssetUserDataInterface)
-	{
-		USD_LOG_WARNING(
-			TEXT("Tried adding AssetUserData to object '%s', but it doesn't implement the AssetUserData interface!"),
-			*Object->GetPathName()
-		);
-		return false;
-	}
-
-	while (AssetUserDataInterface->HasAssetUserDataOfClass(UUsdAssetUserData::StaticClass()))
-	{
-		AssetUserDataInterface->RemoveUserDataOfClass(UUsdAssetUserData::StaticClass());
-	}
-
-	AssetUserDataInterface->AddAssetUserData(AssetUserData);
-	return true;
-}
-PRAGMA_ENABLE_DEPRECATION_WARNINGS
-
 #if USE_USD_SDK
 FString UsdUtils::GetAssetHashPrefix(const pxr::UsdPrim& PrimForAsset, bool bShareAssetsForIdenticalPrims)
 {
@@ -3166,16 +2922,6 @@ bool UsdUtils::RenamePrim(UE::FUsdPrim& Prim, const TCHAR* NewPrimName)
 #endif	  // #if USE_USD_SDK
 }
 
-bool UsdUtils::RemoveNumberedSuffix(FString& Prefix)
-{
-	return UsdUnreal::ObjectUtils::RemoveNumberedSuffix(Prefix);
-}
-
-FString UsdUtils::GetUniqueName(FString Name, const TSet<FString>& UsedNames)
-{
-	return UsdUnreal::ObjectUtils::GetUniqueName(Name, UsedNames);
-}
-
 #if USE_USD_SDK
 FString UsdUtils::GetValidChildName(FString InName, const pxr::UsdPrim& ParentPrim)
 {
diff --git a/Engine/Plugins/Runtime/USDCore/Source/USDUtilities/Private/USDGeomMeshConversion.cpp b/Engine/Plugins/Runtime/USDCore/Source/USDUtilities/Private/USDGeomMeshConversion.cpp
--- a/Engine/Plugins/Runtime/USDCore/Source/USDUtilities/Private/USDGeomMeshConversion.cpp
+++ b/Engine/Plugins/Runtime/USDCore/Source/USDUtilities/Private/USDGeomMeshConversion.cpp
@@ -4039,121 +4039,6 @@ bool UsdToUnreal::ConvertGeomPrimitiveTransform(const pxr::UsdPrim& InPrim, cons
 	return false;
 }
 
-PRAGMA_DISABLE_DEPRECATION_WARNINGS
-UMaterialInstanceDynamic* UsdUtils::CreateDisplayColorMaterialInstanceDynamic(const UsdUtils::FDisplayColorMaterial& DisplayColorDescription)
-{
-	const UUsdProjectSettings* Settings = GetDefault<UUsdProjectSettings>();
-	if (!Settings)
-	{
-		return nullptr;
-	}
-
-	const FSoftObjectPath* ParentPathPtr = nullptr;
-	if (DisplayColorDescription.bHasOpacity)
-	{
-		if (DisplayColorDescription.bIsDoubleSided)
-		{
-			ParentPathPtr = &Settings->ReferenceDisplayColorAndOpacityTwoSidedMaterial;
-		}
-		else
-		{
-			ParentPathPtr = &Settings->ReferenceDisplayColorAndOpacityMaterial;
-		}
-	}
-	else
-	{
-		if (DisplayColorDescription.bIsDoubleSided)
-		{
-			ParentPathPtr = &Settings->ReferenceDisplayColorTwoSidedMaterial;
-		}
-		else
-		{
-			ParentPathPtr = &Settings->ReferenceDisplayColorMaterial;
-		}
-	}
-	if (!ParentPathPtr)
-	{
-		return nullptr;
-	}
-
-	if (UMaterialInterface* ParentMaterial = Cast<UMaterialInterface>(ParentPathPtr->TryLoad()))
-	{
-		FName AssetName = MakeUniqueObjectName(
-			GetTransientPackage(),
-			UMaterialInstanceConstant::StaticClass(),
-			*FString::Printf(
-				TEXT("DisplayColor%s%s"),
-				DisplayColorDescription.bHasOpacity ? TEXT("_Translucent") : TEXT(""),
-				DisplayColorDescription.bIsDoubleSided ? TEXT("_TwoSided") : TEXT("")
-			)
-		);
-
-		if (UMaterialInstanceDynamic* NewMaterial = UMaterialInstanceDynamic::Create(ParentMaterial, GetTransientPackage(), AssetName))
-		{
-			return NewMaterial;
-		}
-	}
-
-	return nullptr;
-}
-
-UMaterialInstanceConstant* UsdUtils::CreateDisplayColorMaterialInstanceConstant(const UsdUtils::FDisplayColorMaterial& DisplayColorDescription)
-{
-#if WITH_EDITOR
-	const UUsdProjectSettings* Settings = GetDefault<UUsdProjectSettings>();
-	if (!Settings)
-	{
-		return nullptr;
-	}
-
-	const FSoftObjectPath* ParentPathPtr = nullptr;
-	if (DisplayColorDescription.bHasOpacity)
-	{
-		if (DisplayColorDescription.bIsDoubleSided)
-		{
-			ParentPathPtr = &Settings->ReferenceDisplayColorAndOpacityTwoSidedMaterial;
-		}
-		else
-		{
-			ParentPathPtr = &Settings->ReferenceDisplayColorAndOpacityMaterial;
-		}
-	}
-	else
-	{
-		if (DisplayColorDescription.bIsDoubleSided)
-		{
-			ParentPathPtr = &Settings->ReferenceDisplayColorTwoSidedMaterial;
-		}
-		else
-		{
-			ParentPathPtr = &Settings->ReferenceDisplayColorMaterial;
-		}
-	}
-
-	if (UMaterialInterface* ParentMaterial = Cast<UMaterialInterface>(ParentPathPtr->TryLoad()))
-	{
-		FName AssetName = MakeUniqueObjectName(
-			GetTransientPackage(),
-			UMaterialInstanceConstant::StaticClass(),
-			*FString::Printf(
-				TEXT("DisplayColor%s%s"),
-				DisplayColorDescription.bHasOpacity ? TEXT("_Translucent") : TEXT(""),
-				DisplayColorDescription.bIsDoubleSided ? TEXT("_TwoSided") : TEXT("")
-			)
-		);
-
-		if (UMaterialInstanceConstant* MaterialInstance = NewObject<UMaterialInstanceConstant>(GetTransientPackage(), AssetName, RF_NoFlags))
-		{
-			UMaterialEditingLibrary::SetMaterialInstanceParent(MaterialInstance, ParentMaterial);
-			return MaterialInstance;
-		}
-	}
-#endif	  // WITH_EDITOR
-	return nullptr;
-}
-
-PRAGMA_ENABLE_DEPRECATION_WARNINGS
-
 UsdUtils::FUsdPrimMaterialAssignmentInfo UsdUtils::GetPrimMaterialAssignments(
 	const pxr::UsdPrim& UsdPrim,
 	const pxr::UsdTimeCode TimeCode,
@@ -4521,53 +4406,6 @@ UsdUtils::FUsdPrimMaterialAssignmentInfo UsdUtils::GetPrimMaterialAssignments(
 	return Result;
 }
 
-TArray<FString> UsdUtils::GetMaterialUsers(const UE::FUsdPrim& MaterialPrim, FName MaterialPurpose)
-{
-	TArray<FString> Result;
-
-	FScopedUsdAllocs Allocs;
-
-	pxr::UsdPrim UsdMaterialPrim{MaterialPrim};
-	if (!UsdMaterialPrim || !UsdMaterialPrim.IsA<pxr::UsdShadeMaterial>())
-	{
-		return Result;
-	}
-
-	pxr::TfToken MaterialPurposeToken = pxr::UsdShadeTokens->allPurpose;
-	if (!MaterialPurpose.IsNone())
-	{
-		MaterialPurposeToken = UnrealToUsd::ConvertToken(*MaterialPurpose.ToString()).Get();
-	}
-
-	pxr::UsdStageRefPtr UsdStage = UsdMaterialPrim.GetStage();
-
-	pxr::UsdPrimRange PrimRange = pxr::UsdPrimRange::Stage(UsdStage, pxr::UsdTraverseInstanceProxies());
-	for (pxr::UsdPrimRange::iterator PrimRangeIt = PrimRange.begin(); PrimRangeIt != PrimRange.end(); ++PrimRangeIt)
-	{
-		pxr::UsdPrim Prim = *PrimRangeIt;
-
-		if (!Prim.HasAPI<pxr::UsdShadeMaterialBindingAPI>())
-		{
-			continue;
-		}
-
-		pxr::UsdShadeMaterialBindingAPI BindingAPI(Prim);
-		pxr::UsdShadeMaterial ShadeMaterial = BindingAPI.ComputeBoundMaterial(MaterialPurposeToken);
-		if (!ShadeMaterial)
-		{
-			continue;
-		}
-
-		pxr::UsdPrim ShadeMaterialPrim = ShadeMaterial.GetPrim();
-		if (ShadeMaterialPrim == UsdMaterialPrim)
-		{
-			Result.Add(UsdToUnreal::ConvertPath(Prim.GetPrimPath()));
-		}
-	}
-
-	return Result;
-}
-
 bool UnrealToUsd::ConvertStaticMesh(
 	const UStaticMesh* StaticMesh,
 	pxr::UsdPrim& UsdPrim,
@@ -6729,18 +6567,6 @@ void UsdUtils::ReplaceUnrealMaterialsWithBaked(
 	TraverseForMaterialReplacement(UsdStage, Root, Empty, {});
 }
 
-void UsdUtils::ReplaceUnrealMaterialsWithBaked(
-	const UE::FUsdStage& Stage,
-	const UE::FSdfLayer& LayerToAuthorIn,
-	const TMap<FString, FString>& BakedMaterials,
-	bool bIsAssetLayer,
-	bool bUsePayload,
-	bool bRemoveUnrealMaterials
-)
-{
-	ReplaceUnrealMaterialsWithBaked(Stage, LayerToAuthorIn, BakedMaterials, bIsAssetLayer, bUsePayload);
-}
-
 namespace UE::UsdGeomMeshConversion::Private
 {
 	template<typename T>
diff --git a/Engine/Plugins/Runtime/USDCore/Source/USDUtilities/Private/USDLightConversion.cpp b/Engine/Plugins/Runtime/USDCore/Source/USDUtilities/Private/USDLightConversion.cpp
--- a/Engine/Plugins/Runtime/USDCore/Source/USDUtilities/Private/USDLightConversion.cpp
+++ b/Engine/Plugins/Runtime/USDCore/Source/USDUtilities/Private/USDLightConversion.cpp
@@ -196,79 +196,6 @@ bool UsdToUnreal::ConvertSphereLight(const pxr::UsdPrim& Prim, UPointLightCompon
 	return true;
 }
 
-bool UsdToUnreal::ConvertDomeLight(
-	const pxr::UsdPrim& Prim,
-	USkyLightComponent& LightComponent,
-	UUsdAssetCache2* TexturesCache,
-	bool bShareAssetsForIdenticalPrims
-)
-{
-	FScopedUsdAllocs UsdAllocs;
-
-	pxr::UsdAttribute TextureFileAttr;
-	if (pxr::UsdLuxDomeLight DomeLight{Prim})
-	{
-		TextureFileAttr = DomeLight.GetTextureFileAttr();
-	}
-	else if (pxr::UsdLuxDomeLight_1 DomeLight1{Prim})
-	{
-		TextureFileAttr = DomeLight1.GetTextureFileAttr();
-	}
-	if (!TextureFileAttr)
-	{
-		return false;
-	}
-
-	// Revert the allocator in case we end up creating a texture on the ansi allocator or something like that
-	FScopedUnrealAllocs UEAllocs;
-
-	const FString ResolvedDomeTexturePath = UsdUtils::GetResolvedAssetPath(TextureFileAttr);
-	if (ResolvedDomeTexturePath.IsEmpty())
-	{
-		FScopedUsdAllocs Allocs;
-
-		pxr::SdfAssetPath TextureAssetPath;
-		TextureFileAttr.Get<pxr::SdfAssetPath>(&TextureAssetPath);
-
-		// Show a good warning for this because it's easy to pick some cubemap asset from the engine (that usually don't come with the
-		// source texture) and have the dome light silently not work again
-		FString TargetAssetPath = UsdToUnreal::ConvertString(TextureAssetPath.GetAssetPath());
-		USD_LOG_USERWARNING(FText::Format(
-			LOCTEXT("MissingDomeLightTexture", "Failed to find texture '{0}' used for DomeLight '{1}'!"),
-			FText::FromString(TargetAssetPath),
-			FText::FromString(UsdToUnreal::ConvertPath(Prim.GetPath()))
-		));
-
-		return true;
-	}
-
-	const FString& DesiredTextureName = FPaths::GetBaseFilename(ResolvedDomeTexturePath);
-
-	// TODO: Ideally we'd expose this
-	const EObjectFlags DesiredFlags = RF_Public | RF_Standalone | RF_Transactional;
-
-	TextureGroup Group = TextureGroup::TEXTUREGROUP_Skybox;
-
-	UTextureCube* Cubemap = nullptr;
-
-	UObject* Outer = GetTransientPackage();
-	FName TextureName = MakeUniqueObjectName(Outer, UTextureCube::StaticClass(), *UsdUnreal::ObjectUtils::SanitizeObjectName(DesiredTextureName));
-	Cubemap = Cast<UTextureCube>(UsdUtils::CreateTexture(ResolvedDomeTexturePath, TextureName, Group, DesiredFlags));
-
-	if (UUsdAssetUserData* TextureUserData = UsdUnreal::ObjectUtils::GetOrCreateAssetUserData(Cubemap))
-	{
-		TextureUserData->PrimPaths.AddUnique(UsdToUnreal::ConvertPath(Prim.GetPath()));
-	}
-
-	if (Cubemap)
-	{
-		LightComponent.Cubemap = Cubemap;
-		LightComponent.SourceType = ESkyLightSourceType::SLS_SpecifiedCubemap;
-	}
-
-	return true;
-}
-
 bool UsdToUnreal::ConvertLuxShapingAPI(const pxr::UsdPrim& Prim, USpotLightComponent& LightComponent, double UsdTimeCode)
 {
 	FScopedUsdAllocs Allocs;
diff --git a/Engine/Plugins/Runtime/USDCore/Source/USDUtilities/Private/USDPrimConversion.cpp b/Engine/Plugins/Runtime/USDCore/Source/USDUtilities/Private/USDPrimConversion.cpp
--- a/Engine/Plugins/Runtime/USDCore/Source/USDUtilities/Private/USDPrimConversion.cpp
+++ b/Engine/Plugins/Runtime/USDCore/Source/USDUtilities/Private/USDPrimConversion.cpp
@@ -6538,11 +6538,6 @@ bool UnrealToUsd::ConvertXformable(
 	return true;
 }
 
-TArray<UE::FUsdAttribute> UnrealToUsd::GetAttributesForProperty(const UE::FUsdPrim& Prim, const FName& PropertyPath)
-{
-	return UsdUtils::GetAttributesForProperty(Prim, PropertyPath);
-}
-
 TArray<UE::FUsdAttribute> UsdUtils::GetAttributesForProperty(const UE::FUsdPrim& Prim, const FName& PropertyPath)
 {
 	using namespace UnrealIdentifiers;
diff --git a/Engine/Plugins/Runtime/USDCore/Source/USDUtilities/Private/USDShadeConversion.cpp b/Engine/Plugins/Runtime/USDCore/Source/USDUtilities/Private/USDShadeConversion.cpp
--- a/Engine/Plugins/Runtime/USDCore/Source/USDUtilities/Private/USDShadeConversion.cpp
+++ b/Engine/Plugins/Runtime/USDCore/Source/USDUtilities/Private/USDShadeConversion.cpp
@@ -2561,44 +2561,6 @@ bool UsdToUnreal::ConvertShadeInputsToParameters(
 	return true;
 }
 
-PRAGMA_DISABLE_DEPRECATION_WARNINGS
-bool UsdToUnreal::ConvertMaterial(
-	const pxr::UsdShadeMaterial& UsdShadeMaterial,
-	UMaterialInstance& Material,
-	UUsdAssetCache2* TexturesCache,
-	const TCHAR* RenderContext,
-	bool bShareAssetsForIdenticalPrims
-)
-{
-	UUsdAssetCache3* NewCache = nullptr;
-	return UsdToUnreal::ConvertMaterial(UsdShadeMaterial, Material, NewCache, RenderContext, bShareAssetsForIdenticalPrims);
-}
-
-bool UsdToUnreal::ConvertMaterial(
-	const pxr::UsdShadeMaterial& UsdShadeMaterial,
-	UMaterial& Material,
-	UUsdAssetCache2* TexturesCache,
-	const TCHAR* RenderContext,
-	bool bShareAssetsForIdenticalPrims
-)
-{
-	UUsdAssetCache3* NewCache = nullptr;
-	return UsdToUnreal::ConvertMaterial(UsdShadeMaterial, Material, NewCache, RenderContext, bShareAssetsForIdenticalPrims);
-}
-
-bool UsdToUnreal::ConvertShadeInputsToParameters(
-	const pxr::UsdShadeMaterial& UsdShadeMaterial,
-	UMaterialInstance& MaterialInstance,
-	UUsdAssetCache2* TexturesCache,
-	const TCHAR* RenderContext,
-	bool bShareAssetsForIdenticalPrims
-)
-{
-	UUsdAssetCache3* NewCache = nullptr;
-	return UsdToUnreal::ConvertShadeInputsToParameters(UsdShadeMaterial, MaterialInstance, NewCache, RenderContext, bShareAssetsForIdenticalPrims);
-}
-PRAGMA_ENABLE_DEPRECATION_WARNINGS
-
 #if WITH_EDITOR
 bool UnrealToUsd::ConvertMaterialToBakedSurface(
 	const UMaterialInterface& InMaterial,
@@ -2707,12 +2669,6 @@ FString UsdUtils::GetResolvedAssetPath(const pxr::UsdAttribute& AssetPathAttr, p
 	return ResolvedPath;
 }
 
-// Deprecated
-FString UsdUtils::GetResolvedTexturePath(const pxr::UsdAttribute& TextureAssetPathAttr)
-{
-	return UsdUtils::GetResolvedAssetPath(TextureAssetPathAttr);
-}
-
 FString UsdUtils::GetTextureHash(
 	const FString& ResolvedTexturePath,
 	bool bSRGB,
@@ -2773,21 +2729,6 @@ FString UsdUtils::GetTextureHash(
 	return LexToString(Hash);
 }
 
-UTexture* UsdUtils::CreateTexture(const pxr::UsdAttribute& TextureAssetPathAttr, const FString& PrimPath, TextureGroup Group, UObject* Outer)
-{
-	const FString ResolvedTexturePath = UsdUtils::GetResolvedAssetPath(TextureAssetPathAttr);
-
-	FName TextureName = MakeUniqueObjectName(
-		Outer,
-		UTexture::StaticClass(),
-		*UsdUnreal::ObjectUtils::SanitizeObjectName(FPaths::GetBaseFilename(ResolvedTexturePath))
-	);
-
-	EObjectFlags Flags = RF_Public | RF_Standalone;
-
-	return UsdUtils::CreateTexture(ResolvedTexturePath, TextureName, Group, Flags, Outer);
-}
-
 UTexture* UsdUtils::CreateTexture(
 	const FString& ResolvedTexturePath,
 	FName SanitizedName,
@@ -3444,35 +3385,6 @@ bool UsdUtils::SetUnrealSurfaceOutput(pxr::UsdPrim& MaterialPrim, const FString&
 	return true;
 }
 
-bool UsdUtils::RemoveUnrealSurfaceOutput(pxr::UsdPrim& MaterialPrim, const UE::FSdfLayer& LayerToAuthorIn)
-{
-	FScopedUsdAllocs UsdAllocs;
-
-	pxr::UsdShadeMaterial ShadeMaterial{MaterialPrim};
-	pxr::UsdShadeConnectableAPI Connectable{MaterialPrim};
-	if (!ShadeMaterial || !Connectable)
-	{
-		return false;
-	}
-
-	if (pxr::UsdShadeOutput MaterialOutput = ShadeMaterial.GetSurfaceOutput(UnrealIdentifiers::Unreal))
-	{
-		if (pxr::UsdShadeShader SurfaceShader = ShadeMaterial.ComputeSurfaceSource(UnrealIdentifiers::Unreal))
-		{
-			// Fully remove the UnrealShader
-			UsdUtils::RemoveAllLocalPrimSpecs(UE::FUsdPrim{SurfaceShader.GetPrim()}, LayerToAuthorIn);
-		}
-
-		// Disconnect would author something like `token outputs:unreal:surface.connect = None`, which is not quite what we want:
-		// That would be an opinion to have it connected to nothing, but instead we just want to remove any opinion whatsoever,
-		// which is what ClearSource does. Note that these will still leave behind `token outputs:unreal:surface` lines,
-		// but those don't actually count as opinions apparently
-		pxr::UsdShadeConnectableAPI::ClearSource(MaterialOutput);
-	}
-
-	return true;
-}
-
 bool UsdUtils::HasSurfaceOutput(const pxr::UsdPrim& MaterialPrim, const FName& RenderContext)
 {
 	if (!MaterialPrim)
diff --git a/Engine/Plugins/Runtime/USDCore/Source/USDUtilities/Private/USDSkeletalDataConversion.cpp b/Engine/Plugins/Runtime/USDCore/Source/USDUtilities/Private/USDSkeletalDataConversion.cpp
--- a/Engine/Plugins/Runtime/USDCore/Source/USDUtilities/Private/USDSkeletalDataConversion.cpp
+++ b/Engine/Plugins/Runtime/USDCore/Source/USDUtilities/Private/USDSkeletalDataConversion.cpp
@@ -2030,29 +2030,6 @@ bool UsdToUnreal::ConvertSkinnedMesh(
 	return true;
 }
 
-bool UsdToUnreal::ConvertSkinnedMesh(
-	const pxr::UsdSkelSkinningQuery& SkinningQuery,
-	const pxr::UsdSkelSkeletonQuery& SkeletonQuery,
-	FSkeletalMeshImportData& SkelMeshImportData,
-	TArray<UsdUtils::FUsdPrimMaterialSlot>& MaterialAssignments,
-	const TMap<FString, TMap<FString, int32>>& MaterialToPrimvarsUVSetNames,
-	const pxr::TfToken& RenderContext,
-	const pxr::TfToken& MaterialPurpose
-)
-{
-	FUsdMeshConversionOptions Options;
-	Options.RenderContext = RenderContext;
-	Options.MaterialPurpose = MaterialPurpose;
-
-	UsdUtils::FUsdPrimMaterialAssignmentInfo TempInfo;
-	TempInfo.Slots = MaterialAssignments;
-
-	const bool bResult = ConvertSkinnedMesh(SkinningQuery, SkeletonQuery, SkelMeshImportData, TempInfo, Options);
-
-	MaterialAssignments = TempInfo.Slots;
-	return bResult;
-}
-
 // Using UsdSkelSkeletonQuery instead of UsdSkelAnimQuery as it automatically does the joint remapping when we ask it to compute joint transforms.
 // It also initializes the joint transforms with the rest pose, if available, in case the animation doesn't provide data for all joints.
 bool UsdToUnreal::ConvertSkelAnim(
@@ -2669,48 +2646,6 @@ bool UsdToUnreal::ConvertBlendShape(
 	return true;
 }
 
-USkeletalMesh* UsdToUnreal::GetSkeletalMeshFromImportData(
-	TArray<FSkeletalMeshImportData>& LODIndexToSkeletalMeshImportData,
-	const TArray<SkeletalMeshImportData::FBone>& InSkeletonBones,
-	UsdUtils::FBlendShapeMap& InBlendShapesByPath,
-	EObjectFlags ObjectFlags,
-	const FName& MeshName,
-	const FName& SkeletonName
-)
-{
-	TRACE_CPUPROFILER_EVENT_SCOPE(UsdToUnreal::GetSkeletalMeshFromImportData);
-
-	FName UniqueMeshName = MakeUniqueObjectName(
-		GetTransientPackage(),
-		USkeletalMesh::StaticClass(),
-		*UsdUnreal::ObjectUtils::SanitizeObjectName(MeshName.ToString())
-	);
-	USkeletalMesh* SkeletalMesh = NewObject<USkeletalMesh>(GetTransientPackage(), UniqueMeshName, ObjectFlags);
-
-	// Generate a Skeleton and associate it to the SkeletalMesh
-	FName UniqueSkeletonName = MakeUniqueObjectName(
-		GetTransientPackage(),
-		USkeleton::StaticClass(),
-		*UsdUnreal::ObjectUtils::SanitizeObjectName(SkeletonName.ToString())
-	);
-	USkeleton* Skeleton = NewObject<USkeleton>(GetTransientPackage(), UniqueSkeletonName, ObjectFlags);
-
-	Skeleton->SetPreviewMesh(SkeletalMesh);
-	SkeletalMesh->SetSkeleton(Skeleton);
-
-	bool bSuccess = ConvertSkeletalImportData(LODIndexToSkeletalMeshImportData, InSkeletonBones, InBlendShapesByPath, SkeletalMesh);
-	if (!bSuccess)
-	{
-		SkeletalMesh->MarkAsGarbage();
-		SkeletalMesh = nullptr;
-
-		Skeleton->MarkAsGarbage();
-		Skeleton = nullptr;
-	}
-
-	return SkeletalMesh;
-}
-
 bool UsdToUnreal::ConvertSkeletalImportData(
 	TArray<FSkeletalMeshImportData>& InLODIndexToSkeletalMeshImportData,
 	const TArray<SkeletalMeshImportData::FBone>& InSkeletonBones,
@@ -3121,17 +3056,6 @@ void UsdUtils::RefreshSkeletalMeshComponent(USkeletalMeshComponent& Component)
 #if USE_USD_SDK
 
 // Adapted from UsdSkel_CacheImpl::ReadScope::_FindOrCreateSkinningQuery because we need to manually create these on UsdGeomMeshes we already have
-pxr::UsdSkelSkinningQuery UsdUtils::CreateSkinningQuery(const pxr::UsdGeomMesh& SkinnedMesh, const pxr::UsdSkelSkeletonQuery& SkeletonQuery)
-{
-	pxr::UsdPrim SkinnedPrim = SkinnedMesh.GetPrim();
-	if (!SkinnedPrim)
-	{
-		return {};
-	}
-
-	return CreateSkinningQuery(SkinnedPrim, SkeletonQuery);
-}
-
 UE::FUsdSkelSkinningQuery UsdUtils::CreateSkinningQuery(const pxr::UsdPrim& SkinnedMeshPrim, const pxr::UsdSkelSkeletonQuery& SkeletonQuery)
 {
 	pxr::UsdSkelBindingAPI SkelBindingAPI{SkinnedMeshPrim};
diff --git a/Engine/Plugins/Runtime/USDCore/Source/USDUtilities/Private/USDTypesConversion.cpp b/Engine/Plugins/Runtime/USDCore/Source/USDUtilities/Private/USDTypesConversion.cpp
--- a/Engine/Plugins/Runtime/USDCore/Source/USDUtilities/Private/USDTypesConversion.cpp
+++ b/Engine/Plugins/Runtime/USDCore/Source/USDUtilities/Private/USDTypesConversion.cpp
@@ -392,24 +392,6 @@ namespace UnrealToUsd
 		return ConvertColor(FLinearColor(InValue));
 	}
 
-	// Deprecated
-	pxr::GfVec2f ConvertVector(const FVector2D& InValue)
-	{
-		return ConvertVectorFloat(InValue);
-	}
-
-	// Deprecated
-	pxr::GfVec3f ConvertVector(const FVector& InValue)
-	{
-		return ConvertVectorFloat(InValue);
-	}
-
-	// Deprecated
-	pxr::GfVec3f ConvertVector(const FUsdStageInfo& StageInfo, const FVector& InValue)
-	{
-		return ConvertVectorFloat(StageInfo, InValue);
-	}
-
 	pxr::GfVec2h ConvertVectorHalf(const FVector2D& InValue)
 	{
 		return pxr::GfVec2h(InValue[0], InValue[1]);
@@ -549,12 +531,6 @@ namespace UnrealToUsd
 		return UsdMatrix;
 	}
 
-	// Deprecated
-	pxr::GfQuatf ConvertQuat(const FQuat& InValue)
-	{
-		return ConvertQuatFloat(InValue);
-	}
-
 	pxr::GfQuath ConvertQuatHalf(const FQuat& InValue)
 	{
 		return pxr::GfQuath{
diff --git a/Engine/Plugins/Runtime/USDCore/Source/USDUtilities/Public/Objects/USDInfoCache.h b/Engine/Plugins/Runtime/USDCore/Source/USDUtilities/Public/Objects/USDInfoCache.h
--- a/Engine/Plugins/Runtime/USDCore/Source/USDUtilities/Public/Objects/USDInfoCache.h
+++ b/Engine/Plugins/Runtime/USDCore/Source/USDUtilities/Public/Objects/USDInfoCache.h
@@ -49,10 +49,6 @@ public:
 	// Retrieves the children of a prim from the cached information
 	UE_API TArray<UE::FSdfPath> GetChildren(const UE::FSdfPath& ParentPath) const;
 
-	// Returns a list of all prims we have generic info about
-	UE_DEPRECATED(5.5, "No longer used")
-	UE_API TSet<UE::FSdfPath> GetKnownPrims() const;
-
 	UE_API void RebuildCacheForSubtree(const UE::FUsdPrim& Prim, FUsdSchemaTranslationContext& Context);
 	UE_API void RebuildCacheForSubtrees(const TArray<UE::FSdfPath>& SubtreeRoots, FUsdSchemaTranslationContext& Context);
 
@@ -89,10 +85,6 @@ public:
 	UE_API TOptional<uint64> GetSubtreeMaterialSlotCount(const UE::FSdfPath& Path);
 	UE_API TOptional<TArray<UsdUtils::FUsdPrimMaterialSlot>> GetSubtreeMaterialSlots(const UE::FSdfPath& Path);
 
-	// Returns true if Path could potentially be collapsed as a Geometry Cache asset
-	UE_DEPRECATED(5.5, "No longer used")
-	UE_API bool IsPotentialGeometryCacheRoot(const UE::FSdfPath& Path) const;
-
 public:
 	// Marks/checks if the provided path to a prototype prim is already being translated.
 	// This is used during scene translation with instanceables, so that the schema translators can early out
@@ -101,45 +93,6 @@ public:
 	UE_API bool IsPrototypeTranslated(const UE::FSdfPath& PrototypePath);
 	UE_API void MarkPrototypeAsTranslated(const UE::FSdfPath& PrototypePath);
 
-public:
-	UE_DEPRECATED(5.5, "Use the UUsdPrimLinkCache object and its analogous function instead")
-	UE_API void LinkAssetToPrim(const UE::FSdfPath& Path, UObject* Asset);
-
-	UE_DEPRECATED(5.5, "Use the UUsdPrimLinkCache object and its analogous function instead")
-	UE_API void UnlinkAssetFromPrim(const UE::FSdfPath& Path, UObject* Asset);
-
-	UE_DEPRECATED(5.5, "Use the UUsdPrimLinkCache object and its analogous function instead")
-	UE_API TArray<TWeakObjectPtr<UObject>> RemoveAllAssetPrimLinks(const UE::FSdfPath& Path);
-
-	UE_DEPRECATED(5.5, "Use the UUsdPrimLinkCache object and its analogous function instead")
-	UE_API TArray<UE::FSdfPath> RemoveAllAssetPrimLinks(const UObject* Asset);
-
-	UE_DEPRECATED(5.5, "Use the UUsdPrimLinkCache object and its analogous function instead")
-	UE_API void RemoveAllAssetPrimLinks();
-
-	UE_DEPRECATED(5.5, "Use the UUsdPrimLinkCache object and its analogous function instead")
-	UE_API TArray<TWeakObjectPtr<UObject>> GetAllAssetsForPrim(const UE::FSdfPath& Path) const;
-
-	template<typename T = UObject>
-	UE_DEPRECATED(5.5, "Use the UUsdPrimLinkCache object and its analogous function instead")
-	T* GetSingleAssetForPrim(const UE::FSdfPath& Path) const
-	{
-		return nullptr;
-	}
-
-	template<typename T>
-	UE_DEPRECATED(5.5, "Use the UUsdPrimLinkCache object and its analogous function instead")
-	TArray<T*> GetAssetsForPrim(const UE::FSdfPath& Path) const
-	{
-		return {};
-	}
-
-	UE_DEPRECATED(5.5, "Use the UUsdPrimLinkCache object and its analogous function instead")
-	UE_API TArray<UE::FSdfPath> GetPrimsForAsset(const UObject* Asset) const;
-
-	UE_DEPRECATED(5.5, "Use the UUsdPrimLinkCache object and its analogous function instead")
-	UE_API TMap<UE::FSdfPath, TArray<TWeakObjectPtr<UObject>>> GetAllAssetPrimLinks() const;
-
 private:
 	friend class FUsdGeomXformableTranslator;
 	friend class FUsdGeometryCacheTranslator;
diff --git a/Engine/Plugins/Runtime/USDCore/Source/USDUtilities/Public/Objects/USDSchemaTranslator.h b/Engine/Plugins/Runtime/USDCore/Source/USDUtilities/Public/Objects/USDSchemaTranslator.h
--- a/Engine/Plugins/Runtime/USDCore/Source/USDUtilities/Public/Objects/USDSchemaTranslator.h
+++ b/Engine/Plugins/Runtime/USDCore/Source/USDUtilities/Public/Objects/USDSchemaTranslator.h
@@ -142,30 +142,8 @@ public:
 	FCreateTranslator CreateFunction;
 };
 
-class UE_DEPRECATED(5.5, "Use the render context functions in USDMaterialUtils.h instead.") USDUTILITIES_API FUsdRenderContextRegistry
-{
-public:
-	FUsdRenderContextRegistry();
-
-	void Register(const FName& RenderContextToken);
-	void Unregister(const FName& RenderContextToken);
-	const TSet<FName>& GetRenderContexts() const;
-	const FName& GetUniversalRenderContext() const;
-	const FName& GetUnrealRenderContext() const;
-};
-
 struct FUsdSchemaTranslationContext : public TSharedFromThis<FUsdSchemaTranslationContext>
 {
-	// Explicitly declare these defaulted special functions or else the compiler will do it elsewhere and emit
-	// deprecated warnings due to usage of bCollapseTopLevelPointInstancers
-	PRAGMA_DISABLE_DEPRECATION_WARNINGS
-	FUsdSchemaTranslationContext(const FUsdSchemaTranslationContext& Other) = default;
-	FUsdSchemaTranslationContext& operator=(const FUsdSchemaTranslationContext& Other) = default;
-	PRAGMA_ENABLE_DEPRECATION_WARNINGS
-
-	UE_DEPRECATED(5.5, "Use the constructor that receives an UUsdAssetCache3 instead")
-	USDUTILITIES_API explicit FUsdSchemaTranslationContext(const UE::FUsdStage& InStage, UUsdAssetCache2& InAssetCache);
-
 	USDUTILITIES_API explicit FUsdSchemaTranslationContext(const UE::FUsdStage& InStage);
 	USDUTILITIES_API explicit FUsdSchemaTranslationContext(const UE::FUsdStage& InStage, UUsdAssetCache3& InAssetCache);
 
@@ -225,16 +203,6 @@ struct FUsdSchemaTranslationContext : public TSharedFromThis<FUsdSchemaTranslati
 	/** Where the translated assets will be stored */
 	TStrongObjectPtr<UUsdAssetCache3> UsdAssetCache;
 
-	PRAGMA_DISABLE_DEPRECATION_WARNINGS
-	/** Where the translated assets will be stored */
-	UE_DEPRECATED(5.5, "Use the 'UsdAssetCache' member instead, which is of the new UUsdAssetCache3 type")
-	TStrongObjectPtr<UUsdAssetCache2> AssetCache;
-
-	/** Caches various information about prims that are expensive to query */
-	UE_DEPRECATED(5.5, "Use the 'UsdInfoCache' member instead, which is of the new UUsdInfoCache type")
-	TSharedPtr<FUsdInfoCache> InfoCache;
-	PRAGMA_ENABLE_DEPRECATION_WARNINGS
-
 	/** Caches various information about prims that are expensive to query */
 	FUsdInfoCache* UsdInfoCache;
 
@@ -281,9 +249,6 @@ struct FUsdSchemaTranslationContext : public TSharedFromThis<FUsdSchemaTranslati
 	 */
 	bool bShareAssetsForIdenticalPrims = true;
 
-	UE_DEPRECATED(5.5, "This property has been renamed to 'Share Assets for Identical Prims'")
-	bool bReuseIdenticalAssets = true;
-
 	/**
 	 * If true, prims with a "LOD" variant set, and "LOD0", "LOD1", etc. variants containing each
 	 * a prim can be parsed into a single UStaticMesh asset with multiple LODs
diff --git a/Engine/Plugins/Runtime/USDCore/Source/USDUtilities/Public/USDConversionUtils.h b/Engine/Plugins/Runtime/USDCore/Source/USDUtilities/Public/USDConversionUtils.h
--- a/Engine/Plugins/Runtime/USDCore/Source/USDUtilities/Public/USDConversionUtils.h
+++ b/Engine/Plugins/Runtime/USDCore/Source/USDUtilities/Public/USDConversionUtils.h
@@ -165,16 +165,6 @@ namespace UsdUtils
 
 	USDUTILITIES_API TUsdStore<pxr::TfToken> GetUVSetName(int32 UVChannelIndex);
 
-	/**
-	 * Heuristic to try and guess what UV index we should assign this primvar to.
-	 * We need something like this because one material may use st0, and another st_0 (both meaning the same thing),
-	 * but a mesh that binds both materials may interpret these as targeting completely different UV sets
-	 * @param PrimvarName - Name of the primvar that should be used as UV set
-	 * @return UV index that should be used for this primvar
-	 */
-	UE_DEPRECATED(5.3, "This heuristic is no longer used. Check GetUVSetPrimvars for the new behavior.")
-	USDUTILITIES_API int32 GetPrimvarUVIndex(FString PrimvarName);
-
 	/**
 	 * Gets the primvars that should be used as UV sets, per index, for this UsdPrim.
 	 * This will return between 0 and MaxNumPrimvars primvars, where the first item of array is the primvar that should be used for
@@ -192,41 +182,6 @@ namespace UsdUtils
 		int32 MaxNumPrimvars = USD_PREVIEW_SURFACE_MAX_UV_SETS
 	);
 
-	/**
-	 * Gets the names of the primvars that should be used as UV sets, per index, for this mesh.
-	 * (e.g. first item of array is primvar for UV set 0, second for UV set 1, etc).
-	 * This overload will only return primvars with 'texcoord2f' role.	 *
-	 * @param UsdMesh - Mesh that contains primvars that can be used as texture coordinates.
-	 * @param MaterialToPrimvarsUVSetNames - Maps from a material prim path, to pairs indicating which primvar names are used as 'st' coordinates, and
-	 * which UVIndex the imported material will sample from (e.g. ["st0", 0], ["myUvSet2", 2], etc). These are supposed to be the materials used by
-	 * the mesh, and we do this because it helps identify which primvars are valid/used as texture coordinates, as the user may have these named as
-	 * 'myUvSet2' and still expect it to work
-	 * @param UsdMeshMaterialAssignmentInfo - Result of calling GetPrimMaterialAssignments on UsdMesh's Prim. This can be provided or will be
-	 * retrieved on-demand using RenderContext
-	 * @param RenderContext - Render context to use when traversing through material shaders looking for used primvars
-	 * @param MaterialPurpose - Specific material purpose to use when parsing the UsdMesh's material bindings
-	 * @return Array where each index gives the primvar that should be used for that UV index
-	 */
-	UE_DEPRECATED(
-		5.3,
-		"Use the signature that just receives an UsdMesh, as the returned primvars no longer depend on material assignment information."
-	)
-	USDUTILITIES_API TArray<TUsdStore<pxr::UsdGeomPrimvar>> GetUVSetPrimvars(
-		const pxr::UsdGeomMesh& UsdMesh,
-		const TMap<FString, TMap<FString, int32>>& MaterialToPrimvarsUVSetNames,
-		const pxr::TfToken& RenderContext = pxr::UsdShadeTokens->universalRenderContext,
-		const pxr::TfToken& MaterialPurpose = pxr::UsdShadeTokens->allPurpose
-	);
-	UE_DEPRECATED(
-		5.3,
-		"Use the signature that just receives an UsdMesh, as the returned primvars no longer depend on material assignment information."
-	)
-	USDUTILITIES_API TArray<TUsdStore<pxr::UsdGeomPrimvar>> GetUVSetPrimvars(
-		const pxr::UsdGeomMesh& UsdMesh,
-		const TMap<FString, TMap<FString, int32>>& MaterialToPrimvarsUVSetNames,
-		const UsdUtils::FUsdPrimMaterialAssignmentInfo& UsdMeshMaterialAssignmentInfo
-	);
-
 	/**
 	 * Rearranges an array of primvars into another array of up to USD_PREVIEW_SURFACE_MAX_UV_SETS primvars, describing which
 	 * primvars should be used for each UV set according to the provided AllowedPrimvarsToUVIndex mapping.
@@ -437,33 +392,6 @@ namespace UsdUtils
 	/** Returns the earliest possible timecode. Use it to always fetch the first frame of an animated attribute */
 	USDUTILITIES_API double GetEarliestTimeCode();
 
-	UE_DEPRECATED(5.5, "This function has been moved to the UsdUnreal::ObjectUtils namespace, from USDObjectUtils.h")
-	USDUTILITIES_API UUsdAssetImportData* GetAssetImportData(UObject* Asset);
-	UE_DEPRECATED(5.5, "This function has been moved to the UsdUnreal::ObjectUtils namespace, from USDObjectUtils.h")
-	USDUTILITIES_API void SetAssetImportData(UObject* Asset, UAssetImportData* ImportData);
-	UE_DEPRECATED(5.5, "This function has been moved to the UsdUnreal::ObjectUtils namespace, from USDObjectUtils.h")
-	USDUTILITIES_API UUsdAssetUserData* GetAssetUserData(const UObject* Object, TSubclassOf<UUsdAssetUserData> Class = {});
-	template<typename T>
-	UE_DEPRECATED(5.5, "This function has been moved to the UsdUnreal::ObjectUtils namespace, from USDObjectUtils.h")
-	inline T* GetAssetUserData(UObject* Object)
-	{
-		PRAGMA_DISABLE_DEPRECATION_WARNINGS
-		return Cast<T>(GetAssetUserData(Object, T::StaticClass()));
-		PRAGMA_ENABLE_DEPRECATION_WARNINGS
-	}
-	UE_DEPRECATED(5.5, "This function has been moved to the UsdUnreal::ObjectUtils namespace, from USDObjectUtils.h")
-	USDUTILITIES_API UUsdAssetUserData* GetOrCreateAssetUserData(UObject* Object, TSubclassOf<UUsdAssetUserData> Class = {});
-	template<typename T>
-	UE_DEPRECATED(5.5, "This function has been moved to the UsdUnreal::ObjectUtils namespace, from USDObjectUtils.h")
-	inline T* GetOrCreateAssetUserData(UObject* Object)
-	{
-		PRAGMA_DISABLE_DEPRECATION_WARNINGS
-		return Cast<T>(GetOrCreateAssetUserData(Object, T::StaticClass()));
-		PRAGMA_ENABLE_DEPRECATION_WARNINGS
-	}
-	UE_DEPRECATED(5.5, "This function has been moved to the UsdUnreal::ObjectUtils namespace, from USDObjectUtils.h")
-	USDUTILITIES_API bool SetAssetUserData(UObject* Object, UUsdAssetUserData* AssetUserData);
-
 #if USE_USD_SDK
 	/**
 	 * Simple utility that generates a prefix that should be added to the hash when using it to cache/query an UsdAssetCache2 object.
@@ -548,23 +476,6 @@ namespace UsdUtils
 	 */
 	USDUTILITIES_API bool RenamePrim(UE::FUsdPrim& Prim, const TCHAR* NewPrimName);
 
-	/**
-	 * Removes any numbered suffix, followed by any number of underscores (e.g. Asset_2, Asset__232_31 or Asset94 all become 'Asset'), making
-	 * sure the string is kept at least one character long. Returns true if it removed anything.
-	 */
-	UE_DEPRECATED(5.5, "This function has been moved to USDObjectUtils.h, within the USDClasses module")
-	USDUTILITIES_API bool RemoveNumberedSuffix(FString& Prefix);
-
-	/**
-	 * Appends numbered suffixes to Name until the result is not contained in UsedNames, and returns it.
-	 * Does not add the result to UsedNames before returning (as it is const).
-	 * @param Name - Received string to make unique (e.g. "MyName")
-	 * @param UsedNames - Strings that cannot be used for the result
-	 * @return Modified Name so that it doesn't match anything in UsedNames (e.g. "MyName" again, or "MyName_0" or "MyName_423")
-	 */
-	UE_DEPRECATED(5.5, "This function has been moved to USDObjectUtils.h, within the USDClasses module")
-	USDUTILITIES_API FString GetUniqueName(FString Name, const TSet<FString>& UsedNames);
-
 #if USE_USD_SDK
 	/**
 	 * Returns a sanitized, valid version of 'InName' that can be used as the name for a child prim of 'ParentPrim'.
diff --git a/Engine/Plugins/Runtime/USDCore/Source/USDUtilities/Public/USDGeomMeshConversion.h b/Engine/Plugins/Runtime/USDCore/Source/USDUtilities/Public/USDGeomMeshConversion.h
--- a/Engine/Plugins/Runtime/USDCore/Source/USDUtilities/Public/USDGeomMeshConversion.h
+++ b/Engine/Plugins/Runtime/USDCore/Source/USDUtilities/Public/USDGeomMeshConversion.h
@@ -50,6 +50,11 @@ namespace UE
 	class FUsdPrim;
 }
 
+namespace UsdUnreal::MaterialUtils
+{
+	struct FDisplayColorMaterial;
+}
+
 #if USE_USD_SDK
 namespace UsdToUnreal
 {
@@ -97,12 +102,6 @@ namespace UsdToUnreal
 		// USD timecode when the prims should be sampled for mesh data and converted
 		pxr::UsdTimeCode TimeCode;
 
-		// Maps from a material prim path, to pairs indicating which primvar names are used as 'st' coordinates for
-		// this mesh, and which UVIndex materials will sample from( e.g.[ "st0", 0 ], [ "myUvSet2", 2 ], etc ).
-		// This is used to pick which primvars will become UV sets.
-		UE_DEPRECATED(5.3, "No longer used, now this information is stored directly on generated Material AssetImportData/AssetUserData")
-		const TMap<FString, TMap<FString, int32>>* MaterialToPrimvarToUVIndex;
-
 		// Whether to try reusing material slots (both local and the ones already in MaterialAssignments) when
 		// converting the material assignments from the mesh
 		bool bMergeIdenticalMaterialSlots;
@@ -322,7 +321,7 @@ namespace UnrealToUsd
 
 namespace UsdUtils
 {
-	using FDisplayColorMaterial = IUsdClassesModule::FDisplayColorMaterial;
+	using FDisplayColorMaterial = UsdUnreal::MaterialUtils::FDisplayColorMaterial;
 
 	/** Describes what type of material assignment a FUsdPrimMaterialSlot represents */
 	enum class EPrimAssignmentType : uint8
@@ -404,18 +403,6 @@ namespace UsdUtils
 		const pxr::UsdTimeCode TimeCode = pxr::UsdTimeCode::EarliestTime()
 	);
 
-	/** Creates a dynamic material instance using the right reference material depending on the given description */
-	PRAGMA_DISABLE_DEPRECATION_WARNINGS
-	UE_DEPRECATED(5.3, "This has moved to the IUsdClassesModule namespace, from USDClassesModule.h")
-	USDUTILITIES_API UMaterialInstanceDynamic* CreateDisplayColorMaterialInstanceDynamic(
-		const UsdUtils::FDisplayColorMaterial& DisplayColorDescription
-	);
-	UE_DEPRECATED(5.3, "This has moved to the IUsdClassesModule namespace, from USDClassesModule.h")
-	USDUTILITIES_API UMaterialInstanceConstant* CreateDisplayColorMaterialInstanceConstant(
-		const UsdUtils::FDisplayColorMaterial& DisplayColorDescription
-	);
-	PRAGMA_ENABLE_DEPRECATION_WARNINGS
-
 	/**
 	 * Extracts all material assignment data from UsdPrim, including material binding, multiple assignment with GeomSubsets, and the unrealMaterial
 	 * custom USD attribute. Guaranteed to return at least one material slot. If UsdPrim is a UsdGeomMesh, it is also guaranteed to have valid
@@ -437,17 +424,6 @@ namespace UsdUtils
 		const pxr::TfToken& MaterialPurpose = pxr::UsdShadeTokens->allPurpose
 	);
 
-	/**
-	 * Returns an array of prim paths to prims within MaterialPrim's stage that have a UsdShadeMaterialBindingAPI, and whose computed material binding
-	 * points to MaterialPrim. The user prim's schema is not even checked, so this list will naturally include regular UsdGeomMesh prims,
-	 * UsdGeomSubset prims, and any other with the UsdShadeMaterialBindingAPI.
-	 */
-	UE_DEPRECATED(5.3, "No longer used as the FUsdInfoCache now retrieves and retains this information")
-	USDUTILITIES_API TArray<FString> GetMaterialUsers(
-		const UE::FUsdPrim& MaterialPrim,
-		FName MaterialPurpose = *UnrealIdentifiers::MaterialAllPurpose
-	);
-
 	/** Returns whether this prim can be interpreted as describing a static mesh with multiple LODs (i.e. if the prim holds the LOD variant set
 	 * itself) */
 	USDUTILITIES_API bool DoesPrimContainMeshLODs(const pxr::UsdPrim& Prim);
@@ -495,15 +471,6 @@ namespace UsdUtils
 		bool bIsAssetLayer,
 		bool bUsePayload
 	);
-	UE_DEPRECATED(5.2, "The bRemoveUnrealMaterials parameter is now deprecated as removing the UE material assignments is no longer needed")
-	USDUTILITIES_API void ReplaceUnrealMaterialsWithBaked(
-		const UE::FUsdStage& Stage,
-		const UE::FSdfLayer& LayerToAuthorIn,
-		const TMap<FString, FString>& BakedMaterials,
-		bool bIsAssetLayer,
-		bool bUsePayload,
-		bool bRemoveUnrealMaterials
-	);
 
 	/**
 	 * Hashes the attributes of the GeomMesh at the given prim path on the Stage at TimeCode
diff --git a/Engine/Plugins/Runtime/USDCore/Source/USDUtilities/Public/USDLightConversion.h b/Engine/Plugins/Runtime/USDCore/Source/USDUtilities/Public/USDLightConversion.h
--- a/Engine/Plugins/Runtime/USDCore/Source/USDUtilities/Public/USDLightConversion.h
+++ b/Engine/Plugins/Runtime/USDCore/Source/USDUtilities/Public/USDLightConversion.h
@@ -88,16 +88,6 @@ namespace UsdToUnreal
 		double UsdTimeCode = UsdUtils::GetDefaultTimeCode()
 	);
 
-	PRAGMA_DISABLE_DEPRECATION_WARNINGS
-	UE_DEPRECATED(5.5, "No longer used: The assets are created and assigned by the FUsdLuxLightTranslator directly")
-	USDUTILITIES_API bool ConvertDomeLight(
-		const pxr::UsdPrim& Prim,
-		USkyLightComponent& LightComponent,
-		UUsdAssetCache2* TexturesCache,
-		bool bShareAssetsForIdenticalPrims = true
-	);
-	PRAGMA_ENABLE_DEPRECATION_WARNINGS
-
 	// These are separately exposed so that they can be reused when reading data into MovieScene tracks.
 	// The other attribute conversions are mostly trivial like a single call to UsdToUnreal::ConvertDistance.
 	USDUTILITIES_API float ConvertLightIntensityAttr(float UsdIntensity, float UsdExposure);
diff --git a/Engine/Plugins/Runtime/USDCore/Source/USDUtilities/Public/USDPrimConversion.h b/Engine/Plugins/Runtime/USDCore/Source/USDUtilities/Public/USDPrimConversion.h
--- a/Engine/Plugins/Runtime/USDCore/Source/USDUtilities/Public/USDPrimConversion.h
+++ b/Engine/Plugins/Runtime/USDCore/Source/USDUtilities/Public/USDPrimConversion.h
@@ -449,9 +449,6 @@ namespace UnrealToUsd
 		TSet<FName>& OutPropertyPathsToRefresh
 	);
 
-	UE_DEPRECATED(5.5, "This function has moved to the UsdUtils namespace")
-	USDUTILITIES_API TArray<UE::FUsdAttribute> GetAttributesForProperty(const UE::FUsdPrim& Prim, const FName& PropertyPath);
-
 	/**
 	 * Converts the properties values from UUsdDrawModeComponent into attribute values of a prim with UsdGeomModelAPI.
 	 * Note that this will even apply the schema itself to UsdPrim if it doesn't already have it, and potentially author
diff --git a/Engine/Plugins/Runtime/USDCore/Source/USDUtilities/Public/USDShadeConversion.h b/Engine/Plugins/Runtime/USDCore/Source/USDUtilities/Public/USDShadeConversion.h
--- a/Engine/Plugins/Runtime/USDCore/Source/USDUtilities/Public/USDShadeConversion.h
+++ b/Engine/Plugins/Runtime/USDCore/Source/USDUtilities/Public/USDShadeConversion.h
@@ -107,7 +107,7 @@ namespace UsdToUnreal
 	 * @param OutMaterial - Output parameter that will be filled with the converted data
 	 * @param TexturesCache - Cache to prevent importing a texture more than once
 	 * @param RenderContext - Which render context output to read from the UsdShadeMaterial
-	 * @param ReuseIdenticalAssets - Whether to reuse identical textures found in the TextureCache or to create dedicated textures for each material
+	 * @param bShareAssetsForIdenticalPrims - Whether to reuse identical textures found in the TextureCache or to create dedicated textures for each material
 	 * @return Whether the conversion was successful or not.
 	 */
 	USDUTILITIES_API bool ConvertMaterial(
@@ -148,33 +148,6 @@ namespace UsdToUnreal
 		const TCHAR* RenderContext = nullptr,
 		bool bShareAssetsForIdenticalPrims = true
 	);
-
-	PRAGMA_DISABLE_DEPRECATION_WARNINGS
-	UE_DEPRECATED(5.5, "Use the other overload that receives an UUsdAssetCache3 object instead")
-	USDUTILITIES_API bool ConvertMaterial(
-		const pxr::UsdShadeMaterial& UsdShadeMaterial,
-		UMaterialInstance& Material,
-		UUsdAssetCache2* TexturesCache = nullptr,
-		const TCHAR* RenderContext = nullptr,
-		bool bShareAssetsForIdenticalPrims = true
-	);
-	UE_DEPRECATED(5.5, "Use the other overload that receives an UUsdAssetCache3 object instead")
-	USDUTILITIES_API bool ConvertMaterial(
-		const pxr::UsdShadeMaterial& UsdShadeMaterial,
-		UMaterial& Material,
-		UUsdAssetCache2* TexturesCache = nullptr,
-		const TCHAR* RenderContext = nullptr,
-		bool bShareAssetsForIdenticalPrims = true
-	);
-	UE_DEPRECATED(5.5, "Use the other overload that receives an UUsdAssetCache3 object instead")
-	USDUTILITIES_API bool ConvertShadeInputsToParameters(
-		const pxr::UsdShadeMaterial& UsdShadeMaterial,
-		UMaterialInstance& MaterialInstance,
-		UUsdAssetCache2* TexturesCache,
-		const TCHAR* RenderContext = nullptr,
-		bool bShareAssetsForIdenticalPrims = true
-	);
-	PRAGMA_ENABLE_DEPRECATION_WARNINGS
 }	 // namespace UsdToUnreal
 
 #if WITH_EDITOR
@@ -261,17 +234,6 @@ namespace UsdUtils
 	 */
 	USDUTILITIES_API bool SetUnrealSurfaceOutput(pxr::UsdPrim& MaterialPrim, const FString& UnrealMaterialPathName);
 
-	/**
-	 * Clears any opinions for the 'unreal' render context surface output of MaterialPrim within LayerToAuthorIn.
-	 * If LayerToAuthorIn is an invalid layer (the default) it will clear opinions from all layers of the stage's layer stack.
-	 *
-	 * @param MaterialPrim - pxr::UsdPrim with the pxr::UsdShadeMaterial schema to update the 'unreal' surface output of
-	 * @param LayerToAuthorIn - Layer to clear the opinions in, or an invalid layer (e.g. UE::FSdfLayer{}, which is the default)
-	 * @return Whether we successfully cleared the opinions or not
-	 */
-	UE_DEPRECATED(5.2, "No longer used as UE material assignments are only visible in the 'unreal' render context anyway")
-	USDUTILITIES_API bool RemoveUnrealSurfaceOutput(pxr::UsdPrim& MaterialPrim, const UE::FSdfLayer& LayerToAuthorIn = UE::FSdfLayer{});
-
 	// Returns whether MaterialPrim is an actual Material, and has a surface output authored for the provided render context
 	USDUTILITIES_API bool HasSurfaceOutput(const pxr::UsdPrim& MaterialPrim, const FName& RenderContext);
 	USDUTILITIES_API bool HasSurfaceOutput(const pxr::UsdPrim& MaterialPrim, const TArray<FString>& RenderContexts);
@@ -304,9 +266,6 @@ namespace UsdUtils
 	/** Returns the resolved path from an pxr::SdfAssetPath attribute. For UDIMs path, returns the path to the 1001 tile. */
 	USDUTILITIES_API FString GetResolvedAssetPath(const pxr::UsdAttribute& AssetPathAttr, pxr::UsdTimeCode TimeCode = pxr::UsdTimeCode::Default());
 
-	UE_DEPRECATED(5.4, "This function has been renamed to 'GetResolvedAssetPath', as it should work for any asset type")
-	USDUTILITIES_API FString GetResolvedTexturePath(const pxr::UsdAttribute& TextureAssetPathAttr);
-
 	/**
 	 * Computes and returns the hash string for the texture at the given path.
 	 * Handles regular texture asset paths as well as asset paths identifying textures inside Usdz archives.
@@ -320,14 +279,6 @@ namespace UsdUtils
 		TextureAddress AddressY
 	);
 
-	UE_DEPRECATED(5.5, "Use the overload that receives the resolved texture path directly.")
-	USDUTILITIES_API UTexture* CreateTexture(
-		const pxr::UsdAttribute& TextureAssetPathAttr,
-		const FString& PrimPath = FString(),
-		TextureGroup Group = TEXTUREGROUP_World,
-		UObject* Outer = GetTransientPackage()
-	);
-
 	USDUTILITIES_API UTexture* CreateTexture(
 		const FString& ResolvedTexturePath,
 		FName SanitizedName,
diff --git a/Engine/Plugins/Runtime/USDCore/Source/USDUtilities/Public/USDSkeletalDataConversion.h b/Engine/Plugins/Runtime/USDCore/Source/USDUtilities/Public/USDSkeletalDataConversion.h
--- a/Engine/Plugins/Runtime/USDCore/Source/USDUtilities/Public/USDSkeletalDataConversion.h
+++ b/Engine/Plugins/Runtime/USDCore/Source/USDUtilities/Public/USDSkeletalDataConversion.h
@@ -134,11 +134,6 @@ namespace UsdUtils
 		const pxr::UsdPrim& SkinnedMeshPrim,
 		const pxr::UsdSkelSkeletonQuery& SkeletonQuery
 	);
-	UE_DEPRECATED(5.5, "Please use the other overload of CreateSkinningQuery, that can better work with the UsdWrapper types and no-RTTI modules")
-	USDUTILITIES_API pxr::UsdSkelSkinningQuery CreateSkinningQuery(
-		const pxr::UsdGeomMesh& SkinnedMesh,
-		const pxr::UsdSkelSkeletonQuery& SkeletonQuery
-	);
 
 	/**
 	 * Sets prim AnimationSource as the animation source for prim Prim.
@@ -297,11 +292,7 @@ namespace UsdToUnreal
 	 * of this function)
 	 * @param SkelMeshImportData - Output parameter that will be filled with the converted data
 	 * @param MaterialAssignments - Output parameter that will be filled with the material assignment data extracted from UsdSkinningQuery
-	 * @param MaterialToPrimvarsUVSetNames - Maps from a material prim path, to pairs indicating which primvar names are used as 'st' coordinates for
-	 * this mesh, and which UVIndex materials will sample from (e.g. ["st0", 0], ["myUvSet2", 2], etc). This is used to pick which primvars will
-	 * become UV sets.
-	 * @param RenderContext - Render context to use when parsing the skinned mesh's materials (e.g. '' for universal, or 'mdl', or 'unreal', etc.)
-	 * @param MaterialPurpose - Material purpose to use when parsing the skinned Mesh prim's material bindings
+	 * @param CommonOptions - Generic mesh conversion options to use for the conversion
 	 * @return Whether the conversion was successful or not.
 	 */
 	USDUTILITIES_API bool ConvertSkinnedMesh(
@@ -311,16 +302,6 @@ namespace UsdToUnreal
 		UsdUtils::FUsdPrimMaterialAssignmentInfo& MaterialAssignments,
 		const FUsdMeshConversionOptions& CommonOptions = FUsdMeshConversionOptions::DefaultOptions
 	);
-	UE_DEPRECATED(5.3, "Please use the overload with an FUsdPrimMaterialAssignmentInfo parameter.")
-	USDUTILITIES_API bool ConvertSkinnedMesh(
-		const pxr::UsdSkelSkinningQuery& UsdSkinningQuery,
-		const pxr::UsdSkelSkeletonQuery& SkeletonQuery,
-		FSkeletalMeshImportData& SkelMeshImportData,
-		TArray<UsdUtils::FUsdPrimMaterialSlot>& MaterialAssignments,
-		const TMap<FString, TMap<FString, int32>>& MaterialToPrimvarsUVSetNames,
-		const pxr::TfToken& RenderContext = pxr::UsdShadeTokens->universalRenderContext,
-		const pxr::TfToken& MaterialPurpose = pxr::UsdShadeTokens->allPurpose
-	);
 
 	/**
 	 * Will extract animation data from the animation source of InUsdSkeletonQuery's skeleton, and populate OutSkeletalAnimationAsset with the data.
@@ -353,26 +334,6 @@ namespace UsdToUnreal
 		float* OutScaleSeconds = nullptr
 	);
 
-	/**
-	 * Builds a USkeletalMesh and USkeleton from the imported data in SkelMeshImportData
-	 * @param LODIndexToSkeletalMeshImportData - Container with the imported skeletal mesh data per LOD level
-	 * @param InSkeletonBones - Bones to use for the reference skeleton (skeleton data on each LODIndexToSkeletalMeshImportData will be ignored).
-	 * @param BlendShapesByPath - Blend shapes to convert to morph targets
-	 * @param ObjectFlags - Flags to use when creating the USkeletalMesh and corresponding USkeleton
-	 * @param MeshName - Name to use for the new USkeletalMesh asset
-	 * @param SkeletonName - Name to use for the new USkeleton asset
-	 * @return Newly created USkeletalMesh, or nullptr in case of failure
-	 */
-	UE_DEPRECATED(5.5, "Prefer manually creating/reusing the SkeletalMesh and Skeleton beforehand and then calling ConvertSkeletalImportData.")
-	USDUTILITIES_API USkeletalMesh* GetSkeletalMeshFromImportData(
-		TArray<FSkeletalMeshImportData>& LODIndexToSkeletalMeshImportData,
-		const TArray<SkeletalMeshImportData::FBone>& InSkeletonBones,
-		UsdUtils::FBlendShapeMap& InBlendShapesByPath,
-		EObjectFlags ObjectFlags,
-		const FName& MeshName = NAME_None,
-		const FName& SkeletonName = NAME_None
-	);
-
 	/**
 	 * Fills in InOutSkeletalMesh (and its Skeleton) with all the provided data converted from USD.
 	 * @param InLODIndexToSkeletalMeshImportData - Container with the imported skeletal mesh data per LOD level
diff --git a/Engine/Plugins/Runtime/USDCore/Source/USDUtilities/Public/USDTypesConversion.h b/Engine/Plugins/Runtime/USDCore/Source/USDUtilities/Public/USDTypesConversion.h
--- a/Engine/Plugins/Runtime/USDCore/Source/USDUtilities/Public/USDTypesConversion.h
+++ b/Engine/Plugins/Runtime/USDCore/Source/USDUtilities/Public/USDTypesConversion.h
@@ -122,13 +122,6 @@ namespace UnrealToUsd
 	/** Assumes the input color is in sRGB space. Returns a color in linear space */
 	USDUTILITIES_API pxr::GfVec4f ConvertColor(const FColor& InValue);
 
-	UE_DEPRECATED(5.4, "Please use either ConvertVectorHalf, ConvertVectorFloat or ConvertVectorDouble.")
-	USDUTILITIES_API pxr::GfVec2f ConvertVector(const FVector2D& InValue);
-	UE_DEPRECATED(5.4, "Please use either ConvertVectorHalf, ConvertVectorFloat or ConvertVectorDouble.")
-	USDUTILITIES_API pxr::GfVec3f ConvertVector(const FVector& InValue);
-	UE_DEPRECATED(5.4, "Please use either ConvertVectorHalf, ConvertVectorFloat or ConvertVectorDouble.")
-	USDUTILITIES_API pxr::GfVec3f ConvertVector(const FUsdStageInfo& StageInfo, const FVector& InValue);
-
 	USDUTILITIES_API pxr::GfVec2h ConvertVectorHalf(const FVector2D& InValue);
 	USDUTILITIES_API pxr::GfVec2f ConvertVectorFloat(const FVector2D& InValue);
 	USDUTILITIES_API pxr::GfVec2d ConvertVectorDouble(const FVector2D& InValue);
@@ -151,9 +144,6 @@ namespace UnrealToUsd
 	USDUTILITIES_API pxr::GfMatrix3d ConvertMatrix(const FMatrix3D& Matrix);
 	USDUTILITIES_API pxr::GfMatrix4d ConvertMatrix(const FMatrix& Matrix);
 
-	UE_DEPRECATED(5.4, "Please use either ConvertQuatHalf, ConvertQuatFloat or ConvertQuatDouble.")
-	USDUTILITIES_API pxr::GfQuatf ConvertQuat(const FQuat& InValue);
-
 	USDUTILITIES_API pxr::GfQuath ConvertQuatHalf(const FQuat& InValue);
 	USDUTILITIES_API pxr::GfQuatf ConvertQuatFloat(const FQuat& InValue);
 	USDUTILITIES_API pxr::GfQuatd ConvertQuatDouble(const FQuat& InValue);
diff --git a/Engine/Plugins/Runtime/USDCore/Source/UnrealUSDWrapper/Private/USDListener.cpp b/Engine/Plugins/Runtime/USDCore/Source/UnrealUSDWrapper/Private/USDListener.cpp
--- a/Engine/Plugins/Runtime/USDCore/Source/UnrealUSDWrapper/Private/USDListener.cpp
+++ b/Engine/Plugins/Runtime/USDCore/Source/UnrealUSDWrapper/Private/USDListener.cpp
@@ -588,9 +588,6 @@ public:
 	virtual ~FUsdListenerImpl();
 
 	FUsdListener::FOnStageEditTargetChanged OnStageEditTargetChanged;
-	PRAGMA_DISABLE_DEPRECATION_WARNINGS
-	FUsdListener::FOnLayersChanged OnLayersChanged;
-	PRAGMA_ENABLE_DEPRECATION_WARNINGS
 	FUsdListener::FOnSdfLayersChanged OnSdfLayersChanged;
 	FUsdListener::FOnSdfLayerDirtinessChanged OnSdfLayerDirtinessChanged;
 	FUsdListener::FOnObjectsChanged OnObjectsChanged;
@@ -656,14 +653,6 @@ FUsdListener::FOnStageEditTargetChanged& FUsdListener::GetOnStageEditTargetChang
 	return Impl->OnStageEditTargetChanged;
 }
 
-PRAGMA_DISABLE_DEPRECATION_WARNINGS
-FUsdListener::FOnLayersChanged& FUsdListener::GetOnLayersChanged()
-{
-	return Impl->OnLayersChanged;
-}
-
-PRAGMA_ENABLE_DEPRECATION_WARNINGS
-
 FUsdListener::FOnSdfLayersChanged& FUsdListener::GetOnSdfLayersChanged()
 {
 	return Impl->OnSdfLayersChanged;
@@ -783,7 +772,7 @@ void FUsdListenerImpl::HandleLayersChangedNotice(const pxr::SdfNotice::LayersDid
 #endif	  // ENABLE_NOTICE_LOGGING
 
 	PRAGMA_DISABLE_DEPRECATION_WARNINGS
-	if ((!OnLayersChanged.IsBound() && !OnSdfLayersChanged.IsBound()) || IsBlocked.GetValue() > 0)
+	if (!OnSdfLayersChanged.IsBound() || IsBlocked.GetValue() > 0)
 	{
 		return;
 	}
@@ -825,7 +814,6 @@ void FUsdListenerImpl::HandleLayersChangedNotice(const pxr::SdfNotice::LayersDid
 #endif	  // ENABLE_NOTICE_LOGGING
 
 	FScopedUnrealAllocs UnrealAllocs;
-	OnLayersChanged.Broadcast(LayersNames);
 	OnSdfLayersChanged.Broadcast(ConvertedLayerToChangeList);
 	PRAGMA_ENABLE_DEPRECATION_WARNINGS
 }
diff --git a/Engine/Plugins/Runtime/USDCore/Source/UnrealUSDWrapper/Private/UsdWrappers/UsdPrim.cpp b/Engine/Plugins/Runtime/USDCore/Source/UnrealUSDWrapper/Private/UsdWrappers/UsdPrim.cpp
--- a/Engine/Plugins/Runtime/USDCore/Source/UnrealUSDWrapper/Private/UsdWrappers/UsdPrim.cpp
+++ b/Engine/Plugins/Runtime/USDCore/Source/UnrealUSDWrapper/Private/UsdWrappers/UsdPrim.cpp
@@ -410,17 +410,6 @@ namespace UE
 #endif	  // #if USE_USD_SDK
 	}
 
-	// Deprecated in 5.3
-	bool FUsdPrim::HasAPI(FName SchemaType, TOptional<FName> InstanceName) const
-	{
-		if (InstanceName.IsSet())
-		{
-			return HasAPI(SchemaType, InstanceName.GetValue());
-		}
-
-		return HasAPI(SchemaType);
-	}
-
 	bool FUsdPrim::CanApplyAPI(FName SchemaIdentifier, FString* OutWhyNot) const
 	{
 		bool bResult = false;
@@ -683,12 +672,6 @@ namespace UE
 #endif	  // #if USE_USD_SDK
 	}
 
-	// Deprecated in 5.3
-	bool FUsdPrim::HasPayload() const
-	{
-		return HasAuthoredPayloads();
-	}
-
 	bool FUsdPrim::IsLoaded() const
 	{
 #if USE_USD_SDK
diff --git a/Engine/Plugins/Runtime/USDCore/Source/UnrealUSDWrapper/Public/USDListener.h b/Engine/Plugins/Runtime/USDCore/Source/UnrealUSDWrapper/Public/USDListener.h
--- a/Engine/Plugins/Runtime/USDCore/Source/UnrealUSDWrapper/Public/USDListener.h
+++ b/Engine/Plugins/Runtime/USDCore/Source/UnrealUSDWrapper/Public/USDListener.h
@@ -123,12 +123,6 @@ public:
 	DECLARE_TS_MULTICAST_DELEGATE(FOnStageEditTargetChanged);
 	UE_API FOnStageEditTargetChanged& GetOnStageEditTargetChanged();
 
-	DECLARE_TS_MULTICAST_DELEGATE_OneParam(FOnLayersChanged, const TArray<FString>&);
-	UE_DEPRECATED(5.3, "Use GetOnSdfLayersChanged")
-	UE_API PRAGMA_DISABLE_DEPRECATION_WARNINGS
-	FOnLayersChanged& GetOnLayersChanged();
-	PRAGMA_ENABLE_DEPRECATION_WARNINGS
-
 	DECLARE_TS_MULTICAST_DELEGATE_OneParam(FOnSdfLayersChanged, const UsdUtils::FLayerToSdfChangeList&);
 	UE_API FOnSdfLayersChanged& GetOnSdfLayersChanged();
 
diff --git a/Engine/Plugins/Runtime/USDCore/Source/UnrealUSDWrapper/Public/UsdWrappers/UsdPrim.h b/Engine/Plugins/Runtime/USDCore/Source/UnrealUSDWrapper/Public/UsdWrappers/UsdPrim.h
--- a/Engine/Plugins/Runtime/USDCore/Source/UnrealUSDWrapper/Public/UsdWrappers/UsdPrim.h
+++ b/Engine/Plugins/Runtime/USDCore/Source/UnrealUSDWrapper/Public/UsdWrappers/UsdPrim.h
@@ -112,12 +112,6 @@ namespace UE
 		bool HasAPI(FName SchemaIdentifier) const;
 		bool HasAPI(FName SchemaIdentifier, FName InstanceName) const;
 
-		UE_DEPRECATED(
-			5.3,
-			"Please use the one FName overload for single apply API schemas, or the two FName overload for multiple apply API schemas."
-		)
-		bool HasAPI(FName SchemaType, TOptional<FName> InstanceName) const;
-
 		bool CanApplyAPI(FName SchemaIdentifier, FString* OutWhyNot = nullptr) const;
 		bool CanApplyAPI(FName SchemaIdentifier, FName InstanceName, FString* OutWhyNot = nullptr) const;
 
@@ -148,8 +142,6 @@ namespace UE
 
 		FUsdPayloads GetPayloads() const;
 		bool HasAuthoredPayloads() const;
-		UE_DEPRECATED(5.3, "Please use HasAuthoredPayloads() instead.")
-		bool HasPayload() const;
 		bool IsLoaded() const;
 		void Load(EUsdLoadPolicy Policy = EUsdLoadPolicy::UsdLoadWithDescendants);
 		void Unload();
