Usd:
- Add support for ue_displacement_attribute to control displacement behavior in MaterialX

#Synced-CL 49985829

diff --git a/Engine/Plugins/Runtime/USDCore/Source/USDUtilities/Private/USDMaterialXShaderGraph.cpp b/Engine/Plugins/Runtime/USDCore/Source/USDUtilities/Private/USDMaterialXShaderGraph.cpp
--- a/Engine/Plugins/Runtime/USDCore/Source/USDUtilities/Private/USDMaterialXShaderGraph.cpp
+++ b/Engine/Plugins/Runtime/USDCore/Source/USDUtilities/Private/USDMaterialXShaderGraph.cpp
@@ -543,12 +543,36 @@ mx::InputPtr FUsdMaterialXShaderGraph::GetInput(mx::NodePtr Node, const pxr::Usd
 	}
 	else // take the input given by the nodedef
 	{
-		mx::NodeDefPtr NodeDef = Node->getNodeDef();
-		mx::InputPtr Input = NodeDef->getActiveInput(InputName);
-		InputType = Input->getType();
+		if (mx::NodeDefPtr NodeDef = Node->getNodeDef())
+		{
+			if (mx::InputPtr Input = NodeDef->getActiveInput(InputName))
+			{
+				InputType = Input->getType();
+			}
+		}
+	}
+
+	mx::InputPtr Input = Node->addInput(InputName, InputType);
+
+	// retrieve the 'ue_displacement_attribute' field needed by MaterialX, in order to set these values in the detail panel instead of scalars in the shader graph
+	// see Engine\Binaries\ThirdParty\MaterialX\libraries\Interchange\Interchange_defs.mtlx for full definition of the attribute
+	if (const pxr::UsdAttribute& Attr = UsdInput.GetAttr())
+	{
+		pxr::UsdMetadataValueMap MetadataValueMap = Attr.GetAllMetadata();
+		if (auto CustomData = MetadataValueMap.find(pxr::TfToken("customData")); CustomData != MetadataValueMap.end())
+		{
+			if (CustomData->second.IsHolding<pxr::VtDictionary>())
+			{
+				pxr::VtDictionary Dictionary = CustomData->second.Get<pxr::VtDictionary>();
+				if (auto DisplacementAttribute = Dictionary.find("ue_displacement_attribute"); DisplacementAttribute != Dictionary.end())
+				{
+					Input->setTypedAttribute(DisplacementAttribute->first, DisplacementAttribute->second.Get<bool>());
+				}
+			}
+		}
 	}
 
-	return Node->addInput(InputName, InputType);
+	return Input;
 }
 
 mx::NodeDefPtr FUsdMaterialXShaderGraph::GetNodeDef(const pxr::UsdShadeConnectionSourceInfo& Source, bool& bIsGeomPropValue)
