USD: Add support for Capsule_1 and Cylinder_1 primitive schemas.

#Synced-CL 49920387

diff --git a/Engine/Plugins/Importers/USDImporter/Source/USDSchemas/Private/USDSchemasModule.cpp b/Engine/Plugins/Importers/USDImporter/Source/USDSchemas/Private/USDSchemasModule.cpp
--- a/Engine/Plugins/Importers/USDImporter/Source/USDSchemas/Private/USDSchemasModule.cpp
+++ b/Engine/Plugins/Importers/USDImporter/Source/USDSchemas/Private/USDSchemasModule.cpp
@@ -46,9 +46,11 @@ public:
 			Registry.Register<FUsdGeomCameraTranslator>(TEXT("UsdGeomCamera")),
 			Registry.Register<FUsdGeomMeshTranslator>(TEXT("UsdGeomMesh")),
 			Registry.Register<FUsdGeomPrimitiveTranslator>(TEXT("UsdGeomCapsule")),
+			Registry.Register<FUsdGeomPrimitiveTranslator>(TEXT("UsdGeomCapsule_1")),
 			Registry.Register<FUsdGeomPrimitiveTranslator>(TEXT("UsdGeomCone")),
 			Registry.Register<FUsdGeomPrimitiveTranslator>(TEXT("UsdGeomCube")),
 			Registry.Register<FUsdGeomPrimitiveTranslator>(TEXT("UsdGeomCylinder")),
+			Registry.Register<FUsdGeomPrimitiveTranslator>(TEXT("UsdGeomCylinder_1")),
 			Registry.Register<FUsdGeomPrimitiveTranslator>(TEXT("UsdGeomPlane")),
 			Registry.Register<FUsdGeomPrimitiveTranslator>(TEXT("UsdGeomSphere")),
 			Registry.Register<FUsdGeomPointInstancerTranslator>(TEXT("UsdGeomPointInstancer")),
diff --git a/Engine/Plugins/Interchange/Extensions/OpenUSD/Source/Import/Private/SchemaHandlers/GprimSchemaHandler.cpp b/Engine/Plugins/Interchange/Extensions/OpenUSD/Source/Import/Private/SchemaHandlers/GprimSchemaHandler.cpp
--- a/Engine/Plugins/Interchange/Extensions/OpenUSD/Source/Import/Private/SchemaHandlers/GprimSchemaHandler.cpp
+++ b/Engine/Plugins/Interchange/Extensions/OpenUSD/Source/Import/Private/SchemaHandlers/GprimSchemaHandler.cpp
@@ -623,9 +623,11 @@ namespace UE::Interchange::USD
 		// UsdGeomBasisCurves, nurbs, points, custom user schemas, etc.
 		return Prim.IsA(TEXT("Mesh")) ||
 			   Prim.IsA(TEXT("Capsule")) ||
+			   Prim.IsA(TEXT("Capsule_1")) ||
 			   Prim.IsA(TEXT("Cone")) ||
 			   Prim.IsA(TEXT("Cube")) ||
 			   Prim.IsA(TEXT("Cylinder")) ||
+			   Prim.IsA(TEXT("Cylinder_1")) ||
 			   Prim.IsA(TEXT("Plane")) ||
 			   Prim.IsA(TEXT("Sphere"));
 	}
diff --git a/Engine/Plugins/Runtime/USDCore/Source/USDUtilities/Private/USDGeomMeshConversion.cpp b/Engine/Plugins/Runtime/USDCore/Source/USDUtilities/Private/USDGeomMeshConversion.cpp
--- a/Engine/Plugins/Runtime/USDCore/Source/USDUtilities/Private/USDGeomMeshConversion.cpp
+++ b/Engine/Plugins/Runtime/USDCore/Source/USDUtilities/Private/USDGeomMeshConversion.cpp
@@ -51,6 +51,8 @@
 #include "USDIncludesStart.h"
 #include "opensubdiv/far/primvarRefiner.h"
 #include "opensubdiv/far/topologyRefiner.h"
+#include "pxr/imaging/geomUtil/capsuleMeshGenerator.h"
+#include "pxr/imaging/geomUtil/cylinderMeshGenerator.h"
 #include "pxr/imaging/pxOsd/meshTopology.h"
 #include "pxr/imaging/pxOsd/refinerFactory.h"
 #include "pxr/usd/sdf/layerUtils.h"
@@ -59,9 +61,11 @@
 #include "pxr/usd/usd/prim.h"
 #include "pxr/usd/usd/primRange.h"
 #include "pxr/usd/usdGeom/capsule.h"
+#include "pxr/usd/usdGeom/capsule_1.h"
 #include "pxr/usd/usdGeom/cone.h"
 #include "pxr/usd/usdGeom/cube.h"
 #include "pxr/usd/usdGeom/cylinder.h"
+#include "pxr/usd/usdGeom/cylinder_1.h"
 #include "pxr/usd/usdGeom/gprim.h"
 #include "pxr/usd/usdGeom/mesh.h"
 #include "pxr/usd/usdGeom/plane.h"
@@ -145,9 +149,82 @@ namespace UE::UsdGeomMeshConversion::Private
 {
 	static const FString DisplayColorID = TEXT("!DisplayColor");
 
-	// Dimensions used when generating Capsule meshes
-	static const float DefaultCapsuleMeshRadius = 0.25;
-	static const float DefaultCapsuleMeshHeight = 0.50;
+	// Like with the other primitive shapes we'll be making a single "default capsule" that we scale around with a transform
+	// to represent capsules of different radii and heights. This lets us animate these attributes on the LevelSequence using a static
+	// mesh, but could lead to distorted shapes for aggressive transforms. Because of that, we use these defaults that specify
+	// a very thin capsule with not many cap segments, as that visually minimizes the distortions we get under these transforms
+	static const double DefaultCapsuleMeshRadius = 0.125;
+	static const double DefaultCapsuleMeshHeight = 0.75;
+	static const double DefaultCapsuleNumRadialSegments = 10;
+	static const double DefaultCapsuleNumCapAxialSegments = 2;
+
+	static const double DefaultConeMeshRadius = 0.5;
+	static const double DefaultConeMeshHeight = 1.0;
+
+	static const double DefaultCubeMeshSize = 1.0;
+
+	static const double DefaultCylinderMeshRadius = 0.50;
+	static const double DefaultCylinderMeshHeight = 1.00;
+	static const double DefaultCylinderNumRadialSegments = 10;
+
+	static const double DefaultSphereMeshRadius = 0.5;
+
+	static const double DefaultPlaneMeshWidth = 1.0;
+	static const double DefaultPlaneMeshLength = 1.0;
+
+	void WarnIfAnimated(pxr::UsdAttribute Attr)
+	{
+		if (Attr.ValueMightBeTimeVarying())
+		{
+			USD_LOG_USERWARNING(FText::Format(
+				LOCTEXT("AnimatedAttributeWarning", "Animation of the '{0}' attribute for prim '{1}' is not supported!"),
+				FText::FromString(UsdToUnreal::ConvertString(Attr.GetName())),
+				FText::FromString(UsdToUnreal::ConvertPath(Attr.GetPrim().GetPrimPath()))
+			));
+		}
+	}
+
+	pxr::TfToken GetAxisWithValidation(pxr::UsdAttribute Attr)
+	{
+		WarnIfAnimated(Attr);
+
+		pxr::TfToken Axis = UsdUtils::GetUsdValue<pxr::TfToken>(Attr, pxr::UsdTimeCode::Default());
+		if (Axis != pxr::UsdGeomTokens->x && Axis != pxr::UsdGeomTokens->y && Axis != pxr::UsdGeomTokens->z)
+		{
+			USD_LOG_USERWARNING(FText::Format(
+				LOCTEXT(
+					"UnsupportedAxis",
+					"Attribute '{0}' has unsupported value '{1}'! Supported values are \"{2}\",\"{3}\" and \"{4}\""
+				),
+				FText::FromString(UsdToUnreal::ConvertPath(Attr.GetPath())),
+				FText::FromString(UsdToUnreal::ConvertToken(Axis)),
+				FText::FromString(UsdToUnreal::ConvertToken(pxr::UsdGeomTokens->x)),
+				FText::FromString(UsdToUnreal::ConvertToken(pxr::UsdGeomTokens->y)),
+				FText::FromString(UsdToUnreal::ConvertToken(pxr::UsdGeomTokens->z))
+			));
+		}
+		return Axis;
+	}
+
+	/**
+	 * Tries fetching the default opinion for the attribute if one is authored, otherwise fetches the EarliestTime opinion.
+	 *
+	 * This is useful because sometimes we want to explicitly fetch the default opinion to place it on the level / use it as a base
+	 * case (like for the primitives), but otherwise check the regular timeSample/spline opinions instead of using the schema fallback values.
+	 */
+	template<typename T>
+	T GetDefaultOrEarliest(const pxr::UsdAttribute& Attr)
+	{
+		pxr::UsdTimeCode TimeCodeToUse = pxr::UsdTimeCode::EarliestTime();
+
+		pxr::UsdResolveInfo Info = Attr.GetResolveInfo(pxr::UsdTimeCode::Default());
+		if (Info.GetSource() == pxr::UsdResolveInfoSourceDefault)
+		{
+			TimeCodeToUse = pxr::UsdTimeCode::Default();
+		}
+
+		return UsdUtils::GetUsdValue<T>(Attr, TimeCodeToUse);
+	};
 
 	int32 GetPrimValueIndex(const pxr::TfToken& InterpType, const int32 VertexIndex, const int32 VertexInstanceIndex, const int32 PolygonIndex)
 	{
@@ -248,10 +325,20 @@ namespace UE::UsdGeomMeshConversion::Private
 		}
 
 		const pxr::PxOsdMeshTopology* Topology = nullptr;
+		TOptional<pxr::PxOsdMeshTopology> CreatedTopology;
+
 		if (pxr::UsdGeomCapsule Capsule = pxr::UsdGeomCapsule{UsdPrim})
 		{
 			Topology = &pxr::UsdImagingGetCapsuleMeshTopology();
 		}
+		else if (pxr::UsdGeomCapsule_1 Capsule_1 = pxr::UsdGeomCapsule_1{UsdPrim})
+		{
+			CreatedTopology = pxr::GeomUtilCapsuleMeshGenerator::GenerateTopology(
+				DefaultCapsuleNumRadialSegments,
+				DefaultCapsuleNumCapAxialSegments
+			);
+			Topology = CreatedTopology.GetPtrOrNull();
+		}
 		else if (pxr::UsdGeomCone Cone = pxr::UsdGeomCone{UsdPrim})
 		{
 			Topology = &pxr::UsdImagingGetUnitConeMeshTopology();
@@ -264,6 +351,11 @@ namespace UE::UsdGeomMeshConversion::Private
 		{
 			Topology = &pxr::UsdImagingGetUnitCylinderMeshTopology();
 		}
+		else if (pxr::UsdGeomCylinder_1 Cylinder_1 = pxr::UsdGeomCylinder_1{UsdPrim})
+		{
+			CreatedTopology = pxr::GeomUtilCylinderMeshGenerator::GenerateTopology(DefaultCylinderNumRadialSegments);
+			Topology = CreatedTopology.GetPtrOrNull();
+		}
 		else if (pxr::UsdGeomSphere Sphere = pxr::UsdGeomSphere{UsdPrim})
 		{
 			Topology = &pxr::UsdImagingGetUnitSphereMeshTopology();
@@ -280,6 +372,37 @@ namespace UE::UsdGeomMeshConversion::Private
 		return {};
 	}
 
+	pxr::VtArray<pxr::GfVec3f> GetTwoRadiiCapsuleMeshPoints(pxr::TfToken Axis, float RadiusBottom, float RadiusTop, float Height)
+	{
+		TOptional<pxr::GfMatrix4d> Transform;
+		if (Axis == pxr::UsdGeomTokens->x)
+		{
+			Transform = pxr::GfMatrix4d{0.0, 1.0, 0.0, 0.0, 0.0, 0.0, 1.0, 0.0, 1.0, 0.0, 0.0, 0.0, 0.0, 0.0, 0.0, 1.0};
+		}
+		else if (Axis == pxr::UsdGeomTokens->y)
+		{
+			Transform = pxr::GfMatrix4d{0.0, 0.0, 1.0, 0.0, 1.0, 0.0, 0.0, 0.0, 0.0, 1.0, 0.0, 0.0, 0.0, 0.0, 0.0, 1.0};
+		}
+
+		const size_t NumPoints = pxr::GeomUtilCapsuleMeshGenerator::ComputeNumPoints(
+			DefaultCapsuleNumRadialSegments,
+			DefaultCapsuleNumCapAxialSegments
+		);
+		pxr::VtArray<pxr::GfVec3f> Points{NumPoints};
+
+		pxr::GeomUtilCapsuleMeshGenerator::GeneratePoints(
+			Points.begin(),
+			DefaultCapsuleNumRadialSegments,
+			DefaultCapsuleNumCapAxialSegments,
+			RadiusBottom,
+			RadiusTop,
+			Height,
+			Transform.GetPtrOrNull()
+		);
+
+		return Points;
+	}
+
 	pxr::VtArray<pxr::GfVec3f> GetUnitCylinderMeshPoints(pxr::TfToken Axis)
 	{
 		if (Axis == pxr::UsdGeomTokens->x)
@@ -318,7 +441,7 @@ namespace UE::UsdGeomMeshConversion::Private
 			}();
 			return YCylinder;
 		}
-		else if (Axis == pxr::UsdGeomTokens->z)
+		else
 		{
 			return pxr::UsdImagingGetUnitCylinderMeshPoints();
 		}
@@ -326,6 +449,33 @@ namespace UE::UsdGeomMeshConversion::Private
 		return {};
 	}
 
+	pxr::VtArray<pxr::GfVec3f> GetTwoRadiiCylinderMeshPoints(pxr::TfToken Axis, float RadiusBottom, float RadiusTop, float Height)
+	{
+		TOptional<pxr::GfMatrix4d> Transform;
+		if (Axis == pxr::UsdGeomTokens->x)
+		{
+			Transform = pxr::GfMatrix4d{0.0, 1.0, 0.0, 0.0, 0.0, 0.0, 1.0, 0.0, 1.0, 0.0, 0.0, 0.0, 0.0, 0.0, 0.0, 1.0};
+		}
+		else if (Axis == pxr::UsdGeomTokens->y)
+		{
+			Transform = pxr::GfMatrix4d{0.0, 0.0, 1.0, 0.0, 1.0, 0.0, 0.0, 0.0, 0.0, 1.0, 0.0, 0.0, 0.0, 0.0, 0.0, 1.0};
+		}
+
+		const size_t NumPoints = pxr::GeomUtilCylinderMeshGenerator::ComputeNumPoints(DefaultCylinderNumRadialSegments);
+		pxr::VtArray<pxr::GfVec3f> Points{NumPoints};
+
+		pxr::GeomUtilCylinderMeshGenerator::GeneratePoints(
+			Points.begin(),
+			DefaultCylinderNumRadialSegments,
+			RadiusBottom,
+			RadiusTop,
+			Height,
+			Transform.GetPtrOrNull()
+		);
+
+		return Points;
+	}
+
 	pxr::VtArray<pxr::GfVec3f> GetUnitConeMeshPoints(pxr::TfToken Axis)
 	{
 		if (Axis == pxr::UsdGeomTokens->x)
@@ -364,7 +514,7 @@ namespace UE::UsdGeomMeshConversion::Private
 			}();
 			return YCylinder;
 		}
-		else if (Axis == pxr::UsdGeomTokens->z)
+		else
 		{
 			return pxr::UsdImagingGetUnitConeMeshPoints();
 		}
@@ -990,9 +1140,11 @@ namespace UE::UsdGeomMeshConversion::Private
 		// Check for primitive types directly (don't just check for Gprim here because if we get something else like a UsdVolVolume or
 		// UsdGeomBasisCurves ConvertGeomPrimitive will return false and we will dump our MeshDescription)
 		else if (Prim.IsA<pxr::UsdGeomCapsule>() ||
+				 Prim.IsA<pxr::UsdGeomCapsule_1>() ||
 				 Prim.IsA<pxr::UsdGeomCone>() ||
 				 Prim.IsA<pxr::UsdGeomCube>() ||
 				 Prim.IsA<pxr::UsdGeomCylinder>() ||
+				 Prim.IsA<pxr::UsdGeomCylinder_1>() ||
 				 Prim.IsA<pxr::UsdGeomPlane>() ||
 				 Prim.IsA<pxr::UsdGeomSphere>())
 		{
@@ -3786,14 +3938,11 @@ bool UsdToUnreal::ConvertGeomPrimitive(
 		// don't try writing (or using non-const operator[]) from PrimitivePoints
 		pxr::VtVec3fArray PrimitivePoints;
 		const pxr::PxOsdMeshTopology* PrimitiveTopology = nullptr;
+		TOptional<pxr::PxOsdMeshTopology> CreatedTopology;
 
 		if (pxr::UsdGeomCapsule Capsule = pxr::UsdGeomCapsule{InPrim})
 		{
-			pxr::TfToken Axis = pxr::UsdGeomTokens->z;
-			if (pxr::UsdAttribute Attr = Capsule.GetAxisAttr())
-			{
-				Axis = UsdUtils::GetUsdValue<pxr::TfToken>(Attr, pxr::UsdTimeCode::Default());
-			}
+			pxr::TfToken Axis = UsdUtils::GetUsdValue<pxr::TfToken>(Capsule.GetAxisAttr(), pxr::UsdTimeCode::Default());
 
 			PrimitivePoints = pxr::UsdImagingGenerateCapsuleMeshPoints(
 				UsdGeomMeshImpl::DefaultCapsuleMeshHeight,
@@ -3802,13 +3951,32 @@ bool UsdToUnreal::ConvertGeomPrimitive(
 			);
 			PrimitiveTopology = &pxr::UsdImagingGetCapsuleMeshTopology();
 		}
+		else if (pxr::UsdGeomCapsule_1 Capsule_1 = pxr::UsdGeomCapsule_1{InPrim})
+		{
+			const double RadiusTop = UsdGeomMeshImpl::GetDefaultOrEarliest<double>(Capsule_1.GetRadiusTopAttr());
+			const double RadiusBottom = UsdGeomMeshImpl::GetDefaultOrEarliest<double>(Capsule_1.GetRadiusBottomAttr());
+			const double RadiusRatio = RadiusTop != 0.0f ? (RadiusBottom / RadiusTop) : 1.0f;
+			const double ActualRadiusTop = UsdGeomMeshImpl::DefaultCapsuleMeshRadius;
+			const double ActualRadiusBottom = ActualRadiusTop * RadiusRatio;
+
+			pxr::TfToken Axis = UsdUtils::GetUsdValue<pxr::TfToken>(Capsule_1.GetAxisAttr(), pxr::UsdTimeCode::Default());
+
+			CreatedTopology = pxr::GeomUtilCapsuleMeshGenerator::GenerateTopology(
+				UsdGeomMeshImpl::DefaultCapsuleNumRadialSegments,
+				UsdGeomMeshImpl::DefaultCapsuleNumCapAxialSegments
+			);
+			PrimitiveTopology = CreatedTopology.GetPtrOrNull();
+
+			PrimitivePoints = UsdGeomMeshImpl::GetTwoRadiiCapsuleMeshPoints(
+				Axis,
+				ActualRadiusBottom,
+				ActualRadiusTop,
+				UsdGeomMeshImpl::DefaultCapsuleMeshHeight
+			);
+		}
 		else if (pxr::UsdGeomCone Cone = pxr::UsdGeomCone{InPrim})
 		{
-			pxr::TfToken Axis = pxr::UsdGeomTokens->z;
-			if (pxr::UsdAttribute Attr = Capsule.GetAxisAttr())
-			{
-				Axis = UsdUtils::GetUsdValue<pxr::TfToken>(Attr, pxr::UsdTimeCode::Default());
-			}
+			pxr::TfToken Axis = UsdUtils::GetUsdValue<pxr::TfToken>(Cone.GetAxisAttr(), pxr::UsdTimeCode::Default());
 
 			PrimitivePoints = UsdGeomMeshImpl::GetUnitConeMeshPoints(Axis);
 			PrimitiveTopology = &pxr::UsdImagingGetUnitConeMeshTopology();
@@ -3820,25 +3988,37 @@ bool UsdToUnreal::ConvertGeomPrimitive(
 		}
 		else if (pxr::UsdGeomCylinder Cylinder = pxr::UsdGeomCylinder{InPrim})
 		{
-			pxr::TfToken Axis = pxr::UsdGeomTokens->z;
-			if (pxr::UsdAttribute Attr = Capsule.GetAxisAttr())
-			{
-				Axis = UsdUtils::GetUsdValue<pxr::TfToken>(Attr, pxr::UsdTimeCode::Default());
-			}
+			pxr::TfToken Axis = UsdUtils::GetUsdValue<pxr::TfToken>(Cylinder.GetAxisAttr(), pxr::UsdTimeCode::Default());
 
 			PrimitivePoints = UsdGeomMeshImpl::GetUnitCylinderMeshPoints(Axis);
 			PrimitiveTopology = &pxr::UsdImagingGetUnitCylinderMeshTopology();
 		}
+		else if (pxr::UsdGeomCylinder_1 Cylinder_1 = pxr::UsdGeomCylinder_1{InPrim})
+		{
+			const double RadiusTop = UsdGeomMeshImpl::GetDefaultOrEarliest<double>(Cylinder_1.GetRadiusTopAttr());
+			const double RadiusBottom = UsdGeomMeshImpl::GetDefaultOrEarliest<double>(Cylinder_1.GetRadiusBottomAttr());
+			const double RadiusRatio = RadiusTop != 0.0f ? (RadiusBottom / RadiusTop) : 1.0f;
+			const double ActualRadiusTop = UsdGeomMeshImpl::DefaultCylinderMeshRadius;
+			const double ActualRadiusBottom = ActualRadiusTop * RadiusRatio;
+
+			pxr::TfToken Axis = UsdUtils::GetUsdValue<pxr::TfToken>(Cylinder_1.GetAxisAttr(), pxr::UsdTimeCode::Default());
+
+			CreatedTopology = pxr::GeomUtilCylinderMeshGenerator::GenerateTopology(UsdGeomMeshImpl::DefaultCylinderNumRadialSegments);
+			PrimitiveTopology = CreatedTopology.GetPtrOrNull();
+
+			PrimitivePoints = UsdGeomMeshImpl::GetTwoRadiiCylinderMeshPoints(
+				Axis,
+				ActualRadiusBottom,
+				ActualRadiusTop,
+				UsdGeomMeshImpl::DefaultCylinderMeshHeight
+			);
+		}
 		else if (pxr::UsdGeomPlane Plane = pxr::UsdGeomPlane{InPrim})
 		{
 			const double Width = 1.0f;
 			const double Length = 1.0f;
 
-			pxr::TfToken Axis = pxr::UsdGeomTokens->z;
-			if (pxr::UsdAttribute Attr = Capsule.GetAxisAttr())
-			{
-				Axis = UsdUtils::GetUsdValue<pxr::TfToken>(Attr, pxr::UsdTimeCode::Default());
-			}
+			pxr::TfToken Axis = UsdUtils::GetUsdValue<pxr::TfToken>(Plane.GetAxisAttr(), pxr::UsdTimeCode::Default());
 
 			PrimitivePoints = pxr::UsdImagingGeneratePlaneMeshPoints(Width, Length, Axis);
 			PrimitiveTopology = &pxr::UsdImagingGetPlaneTopology();
@@ -3899,23 +4079,29 @@ bool UsdToUnreal::ConvertGeomPrimitiveTransform(const pxr::UsdPrim& InPrim, cons
 		}
 	};
 
-	if (pxr::UsdGeomCapsule Capsule = pxr::UsdGeomCapsule{InPrim})
+	if (InPrim.IsA<pxr::UsdGeomCapsule>() || InPrim.IsA<pxr::UsdGeomCapsule_1>())
 	{
-		const double Radius = UsdUtils::GetUsdValue<double>(Capsule.GetRadiusAttr(), InTimeCode);
-		const double Height = UsdUtils::GetUsdValue<double>(Capsule.GetHeightAttr(), InTimeCode);
-
-		pxr::TfToken Axis = pxr::UsdGeomTokens->z;
-		if (pxr::UsdAttribute Attr = Capsule.GetAxisAttr())
+		pxr::UsdAttribute RadiusAttr;
+		pxr::UsdAttribute HeightAttr;
+		pxr::UsdAttribute AxisAttr;
+		if (pxr::UsdGeomCapsule Capsule = pxr::UsdGeomCapsule{InPrim})
 		{
-			Axis = UsdUtils::GetUsdValue<pxr::TfToken>(Attr, pxr::UsdTimeCode::Default());
-			if (Attr.ValueMightBeTimeVarying())
-			{
-				USD_LOG_WARNING(
-					TEXT("Animation of the 'axis' attribute for prim '%s' is not supported!"),
-					*UsdToUnreal::ConvertPath(InPrim.GetPrimPath())
-				);
-			}
+			RadiusAttr = Capsule.GetRadiusAttr();
+			HeightAttr = Capsule.GetHeightAttr();
+			AxisAttr = Capsule.GetAxisAttr();
 		}
+		else if (pxr::UsdGeomCapsule_1 Capsule_1 = pxr::UsdGeomCapsule_1{InPrim})
+		{
+			RadiusAttr = Capsule_1.GetRadiusTopAttr();
+			HeightAttr = Capsule_1.GetHeightAttr();
+			AxisAttr = Capsule_1.GetAxisAttr();
+
+			UsdGeomMeshImpl::WarnIfAnimated(Capsule_1.GetRadiusBottomAttr());
+		}
+
+		const double Radius = UsdUtils::GetUsdValue<double>(RadiusAttr, InTimeCode);
+		const double Height = UsdUtils::GetUsdValue<double>(HeightAttr, InTimeCode);
+		pxr::TfToken Axis = UsdGeomMeshImpl::GetAxisWithValidation(AxisAttr);
 
 		// We use these "Scaling" factors instead of direct height/radius because we're assuming
 		// we'll have generated this capsule mesh using ConvertGeomPrimitive, where we provide
@@ -3936,19 +4122,7 @@ bool UsdToUnreal::ConvertGeomPrimitiveTransform(const pxr::UsdPrim& InPrim, cons
 	{
 		const double Radius = UsdUtils::GetUsdValue<double>(Cone.GetRadiusAttr(), InTimeCode);
 		const double Height = UsdUtils::GetUsdValue<double>(Cone.GetHeightAttr(), InTimeCode);
-
-		pxr::TfToken Axis = pxr::UsdGeomTokens->z;
-		if (pxr::UsdAttribute Attr = Cone.GetAxisAttr())
-		{
-			Axis = UsdUtils::GetUsdValue<pxr::TfToken>(Attr, pxr::UsdTimeCode::Default());
-			if (Attr.ValueMightBeTimeVarying())
-			{
-				USD_LOG_WARNING(
-					TEXT("Animation of the 'axis' attribute for prim '%s' is not supported!"),
-					*UsdToUnreal::ConvertPath(InPrim.GetPrimPath())
-				);
-			}
-		}
+		pxr::TfToken Axis = UsdGeomMeshImpl::GetAxisWithValidation(Cone.GetAxisAttr());
 
 		const double Diameter = 2.0 * Radius;
 		pxr::GfMatrix4d PrimitiveTransform = GetScalingTransform(Axis, Height, Diameter);
@@ -3964,23 +4138,29 @@ bool UsdToUnreal::ConvertGeomPrimitiveTransform(const pxr::UsdPrim& InPrim, cons
 		OutTransform = UsdToUnreal::ConvertMatrix(StageInfo, UsdTransform);
 		return true;
 	}
-	else if (pxr::UsdGeomCylinder Cylinder = pxr::UsdGeomCylinder{InPrim})
+	else if (InPrim.IsA<pxr::UsdGeomCylinder>() || InPrim.IsA<pxr::UsdGeomCylinder_1>())
 	{
-		const double Radius = UsdUtils::GetUsdValue<double>(Cone.GetRadiusAttr(), InTimeCode);
-		const double Height = UsdUtils::GetUsdValue<double>(Cone.GetHeightAttr(), InTimeCode);
-
-		pxr::TfToken Axis = pxr::UsdGeomTokens->z;
-		if (pxr::UsdAttribute Attr = Cone.GetAxisAttr())
+		pxr::UsdAttribute RadiusAttr;
+		pxr::UsdAttribute HeightAttr;
+		pxr::UsdAttribute AxisAttr;
+		if (pxr::UsdGeomCylinder Cylinder = pxr::UsdGeomCylinder{InPrim})
 		{
-			Axis = UsdUtils::GetUsdValue<pxr::TfToken>(Attr, pxr::UsdTimeCode::Default());
-			if (Attr.ValueMightBeTimeVarying())
-			{
-				USD_LOG_WARNING(
-					TEXT("Animation of the 'axis' attribute for prim '%s' is not supported!"),
-					*UsdToUnreal::ConvertPath(InPrim.GetPrimPath())
-				);
-			}
+			RadiusAttr = Cylinder.GetRadiusAttr();
+			HeightAttr = Cylinder.GetHeightAttr();
+			AxisAttr = Cylinder.GetAxisAttr();
 		}
+		else if (pxr::UsdGeomCylinder_1 Cylinder_1 = pxr::UsdGeomCylinder_1{InPrim})
+		{
+			RadiusAttr = Cylinder_1.GetRadiusTopAttr();
+			HeightAttr = Cylinder_1.GetHeightAttr();
+			AxisAttr = Cylinder_1.GetAxisAttr();
+
+			UsdGeomMeshImpl::WarnIfAnimated(Cylinder_1.GetRadiusBottomAttr());
+		}
+
+		const double Radius = UsdUtils::GetUsdValue<double>(RadiusAttr, InTimeCode);
+		const double Height = UsdUtils::GetUsdValue<double>(HeightAttr, InTimeCode);
+		pxr::TfToken Axis = UsdGeomMeshImpl::GetAxisWithValidation(AxisAttr);
 
 		const double Diameter = 2.0 * Radius;
 		pxr::GfMatrix4d PrimitiveTransform = GetScalingTransform(Axis, Height, Diameter);
@@ -3992,19 +4172,7 @@ bool UsdToUnreal::ConvertGeomPrimitiveTransform(const pxr::UsdPrim& InPrim, cons
 	{
 		const double Width = UsdUtils::GetUsdValue<double>(Plane.GetWidthAttr(), InTimeCode);
 		const double Length = UsdUtils::GetUsdValue<double>(Plane.GetLengthAttr(), InTimeCode);
-
-		pxr::TfToken Axis = pxr::UsdGeomTokens->z;
-		if (pxr::UsdAttribute Attr = Plane.GetAxisAttr())
-		{
-			Axis = UsdUtils::GetUsdValue<pxr::TfToken>(Attr, pxr::UsdTimeCode::Default());
-			if (Attr.ValueMightBeTimeVarying())
-			{
-				USD_LOG_WARNING(
-					TEXT("Animation of the 'axis' attribute for prim '%s' is not supported!"),
-					*UsdToUnreal::ConvertPath(InPrim.GetPrimPath())
-				);
-			}
-		}
+		pxr::TfToken Axis = UsdGeomMeshImpl::GetAxisWithValidation(Plane.GetAxisAttr());
 
 		OutTransform = FTransform::Identity;
 
@@ -6091,7 +6259,7 @@ EUsdCollisionType UsdUtils::GetCollisionApproximationType(const pxr::UsdPrim& Us
 	else
 	{
 		// Collision for primitives are converted to their closest approximation
-		if (UsdPrim.IsA(pxr::UsdGeomTokens->Capsule))
+		if (UsdPrim.IsA(pxr::UsdGeomTokens->Capsule) || UsdPrim.IsA(pxr::UsdGeomTokens->Capsule_1))
 		{
 			return EUsdCollisionType::Capsule;
 		}
@@ -6103,7 +6271,7 @@ EUsdCollisionType UsdUtils::GetCollisionApproximationType(const pxr::UsdPrim& Us
 		{
 			return EUsdCollisionType::Cube;
 		}
-		else if (UsdPrim.IsA(pxr::UsdGeomTokens->Cylinder))
+		else if (UsdPrim.IsA(pxr::UsdGeomTokens->Cylinder) || UsdPrim.IsA(pxr::UsdGeomTokens->Cylinder_1))
 		{
 			return EUsdCollisionType::CustomMesh;
 		}
@@ -6899,6 +7067,13 @@ uint64 UsdUtils::GetGprimVertexCount(const pxr::UsdGeomGprim& Gprim, double Time
 		// never going to significantly affect whether a subtree should collapse or not
 		return 82;
 	}
+	else if (pxr::UsdGeomCapsule_1 Capsule_1{Gprim})
+	{
+		return pxr::GeomUtilCapsuleMeshGenerator::ComputeNumPoints(
+			UsdGeomMeshImpl::DefaultCapsuleNumRadialSegments,
+			UsdGeomMeshImpl::DefaultCapsuleNumCapAxialSegments
+		);
+	}
 	else if (pxr::UsdGeomCone Cone{Gprim})
 	{
 		return 31;
@@ -6911,6 +7086,10 @@ uint64 UsdUtils::GetGprimVertexCount(const pxr::UsdGeomGprim& Gprim, double Time
 	{
 		return 42;
 	}
+	else if (pxr::UsdGeomCylinder_1 Cylinder_1{Gprim})
+	{
+		return pxr::GeomUtilCylinderMeshGenerator::ComputeNumPoints(UsdGeomMeshImpl::DefaultCylinderNumRadialSegments);
+	}
 	else if (pxr::UsdGeomSphere Sphere{Gprim})
 	{
 		return 92;
@@ -6925,6 +7104,8 @@ uint64 UsdUtils::GetGprimVertexCount(const pxr::UsdGeomGprim& Gprim, double Time
 
 void UsdUtils::AuthorIdentityTransformGprimAttributes(const pxr::UsdPrim& UsdPrim, bool bDefaultValues, bool bTimeSampleValues)
 {
+	using namespace UE::UsdGeomMeshConversion::Private;
+
 	pxr::UsdGeomGprim Gprim{UsdPrim};
 	if (!Gprim)
 	{
@@ -6960,7 +7141,7 @@ void UsdUtils::AuthorIdentityTransformGprimAttributes(const pxr::UsdPrim& UsdPri
 			const pxr::UsdTimeCode TimeCode{0.0};
 			Attr.Set(Value, TimeCode);
 		}
-	};
+	};	
 
 	// In here we must author the attribute values that cause ConvertGeomPrimitive
 	// to generate meshes in the [-0.5, 0.5] bounding box, as that will correspond to the
@@ -6972,31 +7153,58 @@ void UsdUtils::AuthorIdentityTransformGprimAttributes(const pxr::UsdPrim& UsdPri
 	// function have the attributes generate a scale of 1 instead.
 	if (pxr::UsdGeomCapsule Capsule = pxr::UsdGeomCapsule{UsdPrim})
 	{
-		SetAttrValue(Capsule.CreateRadiusAttr(), 0.25);
-		SetAttrValue(Capsule.CreateHeightAttr(), 0.5);
+		SetAttrValue(Capsule.CreateRadiusAttr(), UsdGeomMeshImpl::DefaultCapsuleMeshRadius);
+		SetAttrValue(Capsule.CreateHeightAttr(), UsdGeomMeshImpl::DefaultCapsuleMeshHeight);
+	}
+	else if (pxr::UsdGeomCapsule_1 Capsule_1 = pxr::UsdGeomCapsule_1{UsdPrim})
+	{
+		pxr::UsdAttribute RadiusTopAttr = Capsule_1.GetRadiusTopAttr();
+
+		// For the two radii primitives we need to preserve the ratio between the radii to retain the right shape,
+		// but we still need to restore RadiusTop to its default value as we'll author our primitive transform based on that fact
+		const double RadiusTop = GetDefaultOrEarliest<double>(Capsule_1.GetRadiusTopAttr());
+		const double RadiusBottom = GetDefaultOrEarliest<double>(Capsule_1.GetRadiusBottomAttr());
+		const double RadiusRatio = RadiusTop != 0.0f ? (RadiusBottom / RadiusTop) : 1.0f;
+		const double NewRadiusTop = UsdGeomMeshImpl::DefaultCapsuleMeshRadius;
+
+		SetAttrValue(Capsule_1.CreateRadiusTopAttr(), NewRadiusTop);
+		SetAttrValue(Capsule_1.CreateRadiusBottomAttr(), NewRadiusTop * RadiusRatio);
+		SetAttrValue(Capsule_1.CreateHeightAttr(), UsdGeomMeshImpl::DefaultCapsuleMeshHeight);
 	}
 	else if (pxr::UsdGeomCone Cone = pxr::UsdGeomCone{UsdPrim})
 	{
-		SetAttrValue(Cone.CreateRadiusAttr(), 0.5);
-		SetAttrValue(Cone.CreateHeightAttr(), 1.0);
+		SetAttrValue(Cone.CreateRadiusAttr(), UsdGeomMeshImpl::DefaultConeMeshRadius);
+		SetAttrValue(Cone.CreateHeightAttr(), UsdGeomMeshImpl::DefaultConeMeshHeight);
 	}
 	else if (pxr::UsdGeomCube Cube = pxr::UsdGeomCube{UsdPrim})
 	{
-		SetAttrValue(Cube.CreateSizeAttr(), 1.0);
+		SetAttrValue(Cube.CreateSizeAttr(), UsdGeomMeshImpl::DefaultCubeMeshSize);
 	}
 	else if (pxr::UsdGeomCylinder Cylinder = pxr::UsdGeomCylinder{UsdPrim})
 	{
-		SetAttrValue(Cylinder.CreateRadiusAttr(), 0.5);
-		SetAttrValue(Cylinder.CreateHeightAttr(), 1.0);
+		SetAttrValue(Cylinder.CreateRadiusAttr(), UsdGeomMeshImpl::DefaultCylinderMeshRadius);
+		SetAttrValue(Cylinder.CreateHeightAttr(), UsdGeomMeshImpl::DefaultCylinderMeshHeight);
+	}
+	else if (pxr::UsdGeomCylinder_1 Cylinder_1 = pxr::UsdGeomCylinder_1{UsdPrim})
+	{
+		// See the Capsule_1 case
+		const double RadiusTop = GetDefaultOrEarliest<double>(Cylinder_1.GetRadiusTopAttr());
+		const double RadiusBottom = GetDefaultOrEarliest<double>(Cylinder_1.GetRadiusBottomAttr());
+		const double RadiusRatio = RadiusTop != 0.0f ? (RadiusBottom / RadiusTop) : 1.0f;
+		const double NewRadiusTop = UsdGeomMeshImpl::DefaultCylinderMeshRadius;
+
+		SetAttrValue(Cylinder_1.CreateRadiusTopAttr(), NewRadiusTop);
+		SetAttrValue(Cylinder_1.CreateRadiusBottomAttr(), NewRadiusTop * RadiusRatio);
+		SetAttrValue(Cylinder_1.CreateHeightAttr(), UsdGeomMeshImpl::DefaultCylinderMeshHeight);
 	}
 	else if (pxr::UsdGeomSphere Sphere = pxr::UsdGeomSphere{UsdPrim})
 	{
-		SetAttrValue(Sphere.CreateRadiusAttr(), 0.5);
+		SetAttrValue(Sphere.CreateRadiusAttr(), UsdGeomMeshImpl::DefaultSphereMeshRadius);
 	}
 	else if (pxr::UsdGeomPlane Plane = pxr::UsdGeomPlane{UsdPrim})
 	{
-		SetAttrValue(Plane.CreateWidthAttr(), 1.0);
-		SetAttrValue(Plane.CreateLengthAttr(), 1.0);
+		SetAttrValue(Plane.CreateWidthAttr(), UsdGeomMeshImpl::DefaultPlaneMeshWidth);
+		SetAttrValue(Plane.CreateLengthAttr(), UsdGeomMeshImpl::DefaultPlaneMeshLength);
 	}
 }
 
diff --git a/Engine/Plugins/Runtime/USDCore/Source/USDUtilities/Private/USDPrimConversion.cpp b/Engine/Plugins/Runtime/USDCore/Source/USDUtilities/Private/USDPrimConversion.cpp
--- a/Engine/Plugins/Runtime/USDCore/Source/USDUtilities/Private/USDPrimConversion.cpp
+++ b/Engine/Plugins/Runtime/USDCore/Source/USDUtilities/Private/USDPrimConversion.cpp
@@ -87,9 +87,11 @@
 #include "pxr/usd/usd/timeCode.h"
 #include "pxr/usd/usdGeom/camera.h"
 #include "pxr/usd/usdGeom/capsule.h"
+#include "pxr/usd/usdGeom/capsule_1.h"
 #include "pxr/usd/usdGeom/cone.h"
 #include "pxr/usd/usdGeom/cube.h"
 #include "pxr/usd/usdGeom/cylinder.h"
+#include "pxr/usd/usdGeom/cylinder_1.h"
 #include "pxr/usd/usdGeom/imageable.h"
 #include "pxr/usd/usdGeom/mesh.h"
 #include "pxr/usd/usdGeom/modelAPI.h"
@@ -6569,6 +6571,11 @@ TArray<UE::FUsdAttribute> UsdUtils::GetAttributesForProperty(const UE::FUsdPrim&
 			Attrs.Emplace(Capsule.GetHeightAttr());
 			Attrs.Emplace(Capsule.GetRadiusAttr());
 		}
+		if (pxr::UsdGeomCapsule_1 Capsule_1{UsdPrim})
+		{
+			Attrs.Emplace(Capsule_1.GetHeightAttr());
+			Attrs.Emplace(Capsule_1.GetRadiusTopAttr());
+		}
 		else if (pxr::UsdGeomCone Cone{UsdPrim})
 		{
 			Attrs.Emplace(Cone.GetHeightAttr());
@@ -6583,6 +6590,11 @@ TArray<UE::FUsdAttribute> UsdUtils::GetAttributesForProperty(const UE::FUsdPrim&
 			Attrs.Emplace(Cylinder.GetHeightAttr());
 			Attrs.Emplace(Cylinder.GetRadiusAttr());
 		}
+		else if (pxr::UsdGeomCylinder_1 Cylinder_1{UsdPrim})
+		{
+			Attrs.Emplace(Cylinder_1.GetHeightAttr());
+			Attrs.Emplace(Cylinder_1.GetRadiusTopAttr());
+		}
 		else if (pxr::UsdGeomPlane Plane{UsdPrim})
 		{
 			Attrs.Emplace(Plane.GetLengthAttr());
diff --git a/Engine/Plugins/Runtime/USDCore/Source/USDUtilities/Public/USDGeomMeshConversion.h b/Engine/Plugins/Runtime/USDCore/Source/USDUtilities/Public/USDGeomMeshConversion.h
--- a/Engine/Plugins/Runtime/USDCore/Source/USDUtilities/Public/USDGeomMeshConversion.h
+++ b/Engine/Plugins/Runtime/USDCore/Source/USDUtilities/Public/USDGeomMeshConversion.h
@@ -204,7 +204,7 @@ namespace UsdToUnreal
 	);
 
 	/**
-	 * Converts an UsdGeomCapsule, UsdGeomCone, UsdGeomCube, UsdGeomCylinder, UsdGeomPlane or UsdGeomSphere and appends
+	 * Converts an UsdGeomCapsule(_1), UsdGeomCone, UsdGeomCube, UsdGeomCylinder(_1), UsdGeomPlane or UsdGeomSphere and appends
 	 * the converted data to InOutMeshDescription and InOutMaterialAssignments.
 	 *
 	 * Note that all generated meshes will have extents within the [-0.5, 0.5] box centered on the origin.
@@ -224,7 +224,7 @@ namespace UsdToUnreal
 	);
 
 	/**
-	 * Extracts a "primitive transform" from the provided UsdGeomCapsule, UsdGeomCone, UsdGeomCube, UsdGeomCylinder,
+	 * Extracts a "primitive transform" from the provided UsdGeomCapsule(_1), UsdGeomCone, UsdGeomCube, UsdGeomCylinder(_1),
 	 * UsdGeomPlane or UsdGeomSphere and places the converted, UE coordinate system transform on OutTransform.
 	 *
 	 * This is intended to be applied to the primitive data generated by ConvertGeomPrimitive, and concatenated with
