USD: Enabled UE allocators for USD modules on Linux and Mac. This does not require rebuilding the USD SDK with explicit UE allocators support since it's already working.
Also enabled it on Windows Arm64.

#Synced-CL 50450691

diff --git a/Engine/Plugins/Runtime/USDCore/Source/UnrealUSDWrapper/Private/UnrealUSDWrapper.cpp b/Engine/Plugins/Runtime/USDCore/Source/UnrealUSDWrapper/Private/UnrealUSDWrapper.cpp
--- a/Engine/Plugins/Runtime/USDCore/Source/UnrealUSDWrapper/Private/UnrealUSDWrapper.cpp
+++ b/Engine/Plugins/Runtime/USDCore/Source/UnrealUSDWrapper/Private/UnrealUSDWrapper.cpp
@@ -1174,9 +1174,12 @@ public:
 #if USE_USD_SDK
 		LLM_SCOPE_BYTAG(Usd);
 
+#if PLATFORM_WINDOWS
 #if !USE_USD_MEMORY_MANAGER
+		// Only Windows need to register the memory overloads when enabled in the USD SDK
 		pxr::MemoryOverloads::RegisterMallocFree(FMemory_Malloc, FMemory_Free);
 #endif
+#endif // PLATFORM_WINDOWS
 		
 		// Path to USD base plugins
 		FString UsdPluginsPath = FPaths::Combine(TEXT(".."), TEXT("ThirdParty"), TEXT("USD"), TEXT("UsdResources"));
diff --git a/Engine/Plugins/Runtime/USDCore/Source/UnrealUSDWrapper/Public/USDMemory.h b/Engine/Plugins/Runtime/USDCore/Source/UnrealUSDWrapper/Public/USDMemory.h
--- a/Engine/Plugins/Runtime/USDCore/Source/UnrealUSDWrapper/Public/USDMemory.h
+++ b/Engine/Plugins/Runtime/USDCore/Source/UnrealUSDWrapper/Public/USDMemory.h
@@ -16,31 +16,37 @@
 
 LLM_DECLARE_TAG_API(Usd, UNREALUSDWRAPPER_API);
 
-#define USE_USD_MEMORY_MANAGER 1
+// Default to 0 in case the USD SDK is not used
+// Also, on Linux/Mac, the USD SDK doesn't need the memory overloads to support UE allocators so the USD memory manager is not needed
+#define USE_USD_MEMORY_MANAGER 0
 
 #if USE_USD_SDK
 #include "USDIncludesStart.h"
 #include "pxr/pxr.h"
 
+// See UnrealUSDWrapper.build.cs for the definition of this macro
 #if ENABLE_USD_MEMORY_OVERLOADS
 
-#ifdef PXR_RUNTIME_MEMORY_OVERLOADS_ENABLED
+#ifndef PXR_RUNTIME_MEMORY_OVERLOADS_ENABLED
+// The USD SDK was compiled without the support for UE memory overloads so the USD memory manager has to be enabled
 #undef USE_USD_MEMORY_MANAGER
-#define USE_USD_MEMORY_MANAGER 0
-#else
+#define USE_USD_MEMORY_MANAGER 1
 COMPILE_WARNING("The USD SDK was compiled without the memory overloads enabled. UE will use its USD Memory Manager instead.")
 #endif // PXR_RUNTIME_MEMORY_OVERLOADS_ENABLED
 
 #else
 
-// Currently supported on Windows x64
-#if PLATFORM_WINDOWS && !PLATFORM_CPU_ARM_FAMILY
+// USD memory overloads not enabled so the USD memory manager has to be enabled on Windows
+#if PLATFORM_WINDOWS
+#undef USE_USD_MEMORY_MANAGER
+#define USE_USD_MEMORY_MANAGER 1
 
+// On Windows the USD SDK is compiled with PXR_ENABLE_RUNTIME_MEMORY_OVERLOADS=ON, so this is just a sanity check
 #ifdef PXR_RUNTIME_MEMORY_OVERLOADS_ENABLED
-COMPILE_ERROR("The USD SDK was compiled with the memory overloads enabled but the UnrealUSDWrapper module was not. Make sure the memory overloads match between UE and the USD SDK.")
+COMPILE_WARNING("The USD SDK was compiled with the memory overloads enabled but the UnrealUSDWrapper module was not. Make sure the memory overloads match between UE and the USD SDK.")
 #endif // PXR_RUNTIME_MEMORY_OVERLOADS_ENABLED
 
-#endif // PLATFORM_WINDOWS && !PLATFORM_CPU_ARM_FAMILY
+#endif // PLATFORM_WINDOWS
 
 #endif // ENABLE_USD_MEMORY_OVERLOADS
 
diff --git a/Engine/Plugins/Runtime/USDCore/Source/UnrealUSDWrapper/UnrealUSDWrapper.Build.cs b/Engine/Plugins/Runtime/USDCore/Source/UnrealUSDWrapper/UnrealUSDWrapper.Build.cs
--- a/Engine/Plugins/Runtime/USDCore/Source/UnrealUSDWrapper/UnrealUSDWrapper.Build.cs
+++ b/Engine/Plugins/Runtime/USDCore/Source/UnrealUSDWrapper/UnrealUSDWrapper.Build.cs
@@ -12,17 +12,18 @@ namespace UnrealBuildTool.Rules
 		// within your module, make sure that you call CheckAndSetupUsdSdk so that your module is properly setup for USD use.
 		public static bool CheckAndSetupUsdSdk(ReadOnlyTargetRules Target, ModuleRules Module)
 		{
-			bool bUseUsdMemoryOverloadsIfSupported = true;
+			// Must be set to true if the USD SDK was compiled with PXR_ENABLE_RUNTIME_MEMORY_OVERLOADS=ON
+			// Only relevant on Windows since on Linux/Mac, the UE allocators will always be used
+			bool bEnableUEAllocatorsInUSD = true;
 			bool bIsUsdSdkEnabled = EnableUsdSdk(Target, Module);
-			bool bEnableUsdMemoryOverloads = (bIsUsdSdkEnabled && bUseUsdMemoryOverloadsIfSupported) &&
-				((Target.Platform == UnrealTargetPlatform.Win64 && Target.Architecture == UnrealArch.X64) //||
-				//Target.Platform == UnrealTargetPlatform.Mac ||
-				//Target.Platform == UnrealTargetPlatform.Linux
-				);
+			
+			// Only Windows require memory overloads in the USD SDK
+			bool bPlatformNeedsUsdMemoryOverloads = Target.Platform == UnrealTargetPlatform.Win64;
+			bool bEnableUsdMemoryOverloads = bIsUsdSdkEnabled && bEnableUEAllocatorsInUSD && bPlatformNeedsUsdMemoryOverloads;
 			Module.PublicDefinitions.Add("ENABLE_USD_MEMORY_OVERLOADS=" + (bEnableUsdMemoryOverloads ? "1" : "0"));
 
 			Module.PublicDefinitions.Add("USD_MERGED_MODULES=" + (Target.bMergeModules ? "1" : "0"));
-			if (!Target.bMergeModules && !bEnableUsdMemoryOverloads)
+			if (!Target.bMergeModules && !bEnableUsdMemoryOverloads && bPlatformNeedsUsdMemoryOverloads)
 			{
 				Module.PrivateDefinitions.Add("SUPPRESS_PER_MODULE_INLINE_FILE"); // This module does not use core's standard operator new/delete overloads
 			}
