USD: Add support for the light inputs:normalize attribute, by using normalized light units for the spawned components.

#Synced-CL 50488581

diff --git a/Engine/Plugins/Interchange/Extensions/OpenUSD/Source/Import/Private/SchemaHandlers/LightSchemaHandler.cpp b/Engine/Plugins/Interchange/Extensions/OpenUSD/Source/Import/Private/SchemaHandlers/LightSchemaHandler.cpp
--- a/Engine/Plugins/Interchange/Extensions/OpenUSD/Source/Import/Private/SchemaHandlers/LightSchemaHandler.cpp
+++ b/Engine/Plugins/Interchange/Extensions/OpenUSD/Source/Import/Private/SchemaHandlers/LightSchemaHandler.cpp
@@ -23,6 +23,7 @@
 #include "InterchangeTexture2DNode.h"
 #include "Nodes/InterchangeBaseNodeContainer.h"
 
+#include "Engine/Scene.h"
 #include "Misc/Paths.h"
 
 #include "USDIncludesStart.h"
@@ -108,13 +109,18 @@ namespace UE::LuxLightSchemaHandler::Private
 
 		// Ref. UsdToUnreal::ConvertLight
 		static const FString IntensityToken = UsdToUnreal::ConvertToken(pxr::UsdLuxTokens->inputsIntensity);
-		float Intensity = 1.0f;
-		bool bIntensityAuthored = UsdUtils::GetAuthoredAttributeValue<float>(Prim, IntensityToken, Intensity);
+		float IntensityUsd = 1.0f;
+		bool bIntensityAuthored = UsdUtils::GetAuthoredAttributeValue<float>(Prim, IntensityToken, IntensityUsd);
 
 		static const FString ExposureToken = UsdToUnreal::ConvertToken(pxr::UsdLuxTokens->inputsExposure);
 		float Exposure = 0.0f;
 		UsdUtils::GetAuthoredAttributeValue<float>(Prim, ExposureToken, Exposure);
 
+		static const FString NormalizeToken = UsdToUnreal::ConvertToken(pxr::UsdLuxTokens->inputsNormalize);
+		bool bNormalizedIntensityUsd = false;
+		bool bNormalizeAuthored = UsdUtils::GetAuthoredAttributeValue<bool>(Prim, NormalizeToken, bNormalizedIntensityUsd);
+		bNormalizedIntensityUsd = bNormalizeAuthored && bNormalizedIntensityUsd;
+
 		static const FString RadiusToken = UsdToUnreal::ConvertToken(pxr::UsdLuxTokens->inputsRadius);
 
 		// Temporarily a bit awkward here as we updated our VtValue::Get() to not perform any type coercion, so we
@@ -171,8 +177,9 @@ namespace UE::LuxLightSchemaHandler::Private
 			LightNode->SetCustomUseTemperature(UseTemperature);
 		}
 
-		// "Shadow enabled" currently not supported
-
+		float SolidAngle = 1.0f;
+		float AreaSqMeters = 1.0f;
+		bool bAllowNormalizingIntensity = true;
 		if (Prim.IsA(TEXT("DistantLight")))
 		{
 			UInterchangeDirectionalLightNode* CastLightNode = Cast<UInterchangeDirectionalLightNode>(LightNode);
@@ -181,29 +188,27 @@ namespace UE::LuxLightSchemaHandler::Private
 				return;
 			}
 
-			if (bIntensityAuthored)
-			{
-				Intensity = UsdToUnreal::ConvertLightIntensityAttr(Intensity, Exposure);
-				LightNode->SetCustomIntensity(Intensity);
-			}
+			// For distant lights we always use Lux, for which normalization by area doesn't make much sense (it's already
+			// a per-area sort of unit)
+			bAllowNormalizingIntensity = false;
 
 			static const FString SourceAngleToken = UsdToUnreal::ConvertToken(pxr::UsdLuxTokens->inputsAngle);
-			if (float Angle; UsdUtils::GetAuthoredAttributeValue<float>(Prim, SourceAngleToken, Angle))
+			if (float UsdAngle; UsdUtils::GetAuthoredAttributeValue<float>(Prim, SourceAngleToken, UsdAngle))
 			{
-				CastLightNode->SetCustomLightSourceAngle(Angle);
+				CastLightNode->SetCustomLightSourceAngle(UsdAngle);
 			}
 		}
 		else if (Prim.IsA(TEXT("SphereLight")))
 		{
 			const FUsdStageInfo StageInfo(Prim.GetStage());
 
-			float Radius = 1.0f;
+			float RadiusUE = 1.0f;
 			if (UInterchangePointLightNode* CastLightNode = Cast<UInterchangePointLightNode>(LightNode))
 			{
-				if (UsdUtils::GetAuthoredAttributeValue<float>(Prim, RadiusToken, Radius))
+				if (float RadiusUsd = 1.0f; UsdUtils::GetAuthoredAttributeValue<float>(Prim, RadiusToken, RadiusUsd))
 				{
-					Radius = UsdToUnreal::ConvertDistance(StageInfo, Radius);
-					CastLightNode->SetCustomSourceRadius(Radius);
+					RadiusUE = UsdToUnreal::ConvertDistance(StageInfo, RadiusUsd);
+					CastLightNode->SetCustomSourceRadius(RadiusUE);
 				}
 			}
 
@@ -218,32 +223,31 @@ namespace UE::LuxLightSchemaHandler::Private
 				static const FString ConeAngleToken = UsdToUnreal::ConvertToken(pxr::UsdLuxTokens->inputsShapingConeAngle);
 				static const FString ConeSoftnessToken = UsdToUnreal::ConvertToken(pxr::UsdLuxTokens->inputsShapingConeSoftness);
 
-				float ConeAngle = 90.0f;
-				bool bConeAngleAuthored = UsdUtils::GetAuthoredAttributeValue<float>(Prim, ConeAngleToken, ConeAngle);
+				// Default from UsdLuxShapingAPI
+				float UsdConeAngle = 90.0f;
+				bool bConeAngleAuthored = UsdUtils::GetAuthoredAttributeValue<float>(Prim, ConeAngleToken, UsdConeAngle);
 
-				float ConeSoftness = 0.0f;
-				bool bConeSoftnessAuthored = UsdUtils::GetAuthoredAttributeValue<float>(Prim, ConeSoftnessToken, ConeSoftness);
+				// Default from UsdLuxShapingAPI
+				float UsdConeSoftness = 0.0f;
+				bool bConeSoftnessAuthored = UsdUtils::GetAuthoredAttributeValue<float>(Prim, ConeSoftnessToken, UsdConeSoftness);
+
+				float InnerConeAngle = 0.0f;
+				float OuterConeAngle = 44.0f; // Default from the USpotLightComponent constructor
+				UsdToUnreal::ConvertSpotLightConeAngles(UsdConeAngle, UsdConeSoftness, OuterConeAngle, InnerConeAngle);
+
+				SolidAngle = UsdUtils::GetSpotLightSolidAngle(OuterConeAngle, InnerConeAngle);
+				AreaSqMeters = UsdUtils::GetSpotLightArea(RadiusUE / 100.0f);
 
 				if (bConeAngleAuthored || bConeSoftnessAuthored)
 				{
-					float InnerConeAngle = 0.0f;
-					const float OuterConeAngle = UsdToUnreal::ConvertConeAngleSoftnessAttr(ConeAngle, ConeSoftness, InnerConeAngle);
 					CastLightNode->SetCustomInnerConeAngle(InnerConeAngle);
 					CastLightNode->SetCustomOuterConeAngle(OuterConeAngle);
 				}
 			}
-
-			UInterchangePointLightNode* CastLightNode = Cast<UInterchangePointLightNode>(LightNode);
-			if (!ensure(CastLightNode))
-			{
-				return;
-			}
-
-			if (bIntensityAuthored)
+			else
 			{
-				Intensity = UsdToUnreal::ConvertLightIntensityAttr(Intensity, Exposure);
-				CastLightNode->SetCustomIntensityUnits(EInterchangeLightUnits::Nits);
-				CastLightNode->SetCustomIntensity(Intensity);
+				SolidAngle = UsdUtils::GetPointLightSolidAngle();
+				AreaSqMeters = UsdUtils::GetSpotLightArea(RadiusUE / 100.0f);
 			}
 		}
 		else if (Prim.IsA(TEXT("RectLight")) || Prim.IsA(TEXT("DiskLight")))
@@ -261,44 +265,36 @@ namespace UE::LuxLightSchemaHandler::Private
 
 			if (Prim.IsA(TEXT("RectLight")))
 			{
-				float Width = 1.0f;
-				if (UsdUtils::GetAuthoredAttributeValue<float>(Prim, WidthToken, Width))
+				float WidthUE = 1.0f;
+				if (float WidthUsd = 1.0f; UsdUtils::GetAuthoredAttributeValue<float>(Prim, WidthToken, WidthUsd))
 				{
-					Width = UsdToUnreal::ConvertDistance(StageInfo, Width);
-					CastLightNode->SetCustomSourceWidth(Width);
+					WidthUE = UsdToUnreal::ConvertDistance(StageInfo, WidthUsd);
+					CastLightNode->SetCustomSourceWidth(WidthUE);
 				}
 
-				float Height = 1.0f;
-				if (UsdUtils::GetAuthoredAttributeValue<float>(Prim, HeightToken, Height))
+				float HeightUE = 1.0f;
+				if (float HeightUsd = 1.0f; UsdUtils::GetAuthoredAttributeValue<float>(Prim, HeightToken, HeightUsd))
 				{
-					Height = UsdToUnreal::ConvertDistance(StageInfo, Height);
-					CastLightNode->SetCustomSourceHeight(Height);
+					HeightUE = UsdToUnreal::ConvertDistance(StageInfo, HeightUsd);
+					CastLightNode->SetCustomSourceHeight(HeightUE);
 				}
 
-				if (bIntensityAuthored)
-				{
-					Intensity = UsdToUnreal::ConvertLightIntensityAttr(Intensity, Exposure);
-					CastLightNode->SetCustomIntensityUnits(EInterchangeLightUnits::Nits);
-					CastLightNode->SetCustomIntensity(Intensity);
-				}
+				SolidAngle = UsdUtils::GetRectLightSolidAngle();
+				AreaSqMeters = UsdUtils::GetRectLightArea(WidthUE / 100.0f, HeightUE / 100.0f);
 			}
 			else // DiskLight
 			{
-				float Radius = 0.5f;
-				if (UsdUtils::GetAuthoredAttributeValue<float>(Prim, RadiusToken, Radius))
+				float RadiusUE = 0.5f;
+				if (float RadiusUsd = 1.0f; UsdUtils::GetAuthoredAttributeValue<float>(Prim, RadiusToken, RadiusUsd))
 				{
-					float Width = UsdToUnreal::ConvertDistance(StageInfo, Radius) * 2.f;
-					float Height = Width;
-					CastLightNode->SetCustomSourceWidth(Width);
-					CastLightNode->SetCustomSourceHeight(Height);
+					RadiusUE = UsdToUnreal::ConvertDistance(StageInfo, RadiusUsd);
+					float SquareSideUE = UsdUtils::GetRectLightSideFromDiskLight(RadiusUE);
+					CastLightNode->SetCustomSourceWidth(SquareSideUE);
+					CastLightNode->SetCustomSourceHeight(SquareSideUE);
 				}
 
-				if (bIntensityAuthored)
-				{
-					Intensity = UsdToUnreal::ConvertLightIntensityAttr(Intensity, Exposure);
-					CastLightNode->SetCustomIntensityUnits(EInterchangeLightUnits::Nits);
-					CastLightNode->SetCustomIntensity(Intensity);
-				}
+				SolidAngle = UsdUtils::GetDiskLightSolidAngle();
+				AreaSqMeters = UsdUtils::GetDiskLightArea(RadiusUE / 100.0f);
 			}
 		}
 		else if (Prim.IsA(TEXT("DomeLight")) || Prim.IsA(TEXT("DomeLight_1")))
@@ -309,13 +305,9 @@ namespace UE::LuxLightSchemaHandler::Private
 				return;
 			}
 
-			// This doesn't translate very well as SkyLightComponents seem to have an "IntensityScale", and there's never any clear
-			// indication of what the dome light intensity actually is... But I don't think there's much we can do
-			if (bIntensityAuthored)
-			{
-				Intensity = UsdToUnreal::ConvertLightIntensityAttr(Intensity, Exposure);
-				LightNode->SetCustomIntensity(Intensity);
-			}
+			// For dome lights the intensity is just a simple multiplier with no physical meaning, so normalization
+			// doesn't make much sense
+			bAllowNormalizingIntensity = false;
 
 			// Show some warning in case the file specifies a different DomeLight format, as we only support latlong
 			static const FString FormatToken = UsdToUnreal::ConvertToken(pxr::UsdLuxTokens->inputsTextureFormat);
@@ -367,6 +359,40 @@ namespace UE::LuxLightSchemaHandler::Private
 				));
 			}
 		}
+
+		if (bIntensityAuthored)
+		{
+			// This will end up being Nits for local lights (point, rect, disk, spot), lux for distant lights,
+			// and a simple multiplier without units for dome lights
+			float Intensity = UsdToUnreal::ConvertLightIntensity(IntensityUsd, Exposure);
+
+			// We only normalize the intensity for local lights here as we can flip them into Lumen to represent
+			// normalized units. For distant/dome lights we just ignore normalization
+			//
+			// It's a bit awkward because USD does define normalization for distant lights at least, but the normalization
+			// USD suggests involves mainly hard-coding a sizeFactor constant (which we won't do as there's no good place
+			// to put it)
+			float IntensityUE = 1.0;
+			EInterchangeLightUnits UnitsUE = EInterchangeLightUnits::Unitless;
+			if (bNormalizedIntensityUsd && bAllowNormalizingIntensity)
+			{
+				IntensityUE = UsdUtils::ConvertIntensityFromNits(Intensity, SolidAngle, AreaSqMeters, ELightUnits::Lumens);
+				UnitsUE = EInterchangeLightUnits::Lumens;
+			}
+			else
+			{
+				IntensityUE = Intensity;
+				UnitsUE = EInterchangeLightUnits::Nits;
+			}
+
+			LightNode->SetCustomIntensity(IntensityUE);
+			if (UInterchangeLightNode* DerivedLightNode = Cast<UInterchangeLightNode>(LightNode))
+			{
+				// Note that directional and sky lights are not LightNodes (they derive BaseLightNode directly), so this
+				// check prevents us from trying to set them with Nits
+				DerivedLightNode->SetCustomIntensityUnits(UnitsUE);
+			}
+		}
 	}
 }	 // namespace UE::LuxLightSchemaHandler::Private
 
diff --git a/Engine/Plugins/Runtime/USDCore/Source/USDUtilities/Private/USDLightConversion.cpp b/Engine/Plugins/Runtime/USDCore/Source/USDUtilities/Private/USDLightConversion.cpp
--- a/Engine/Plugins/Runtime/USDCore/Source/USDUtilities/Private/USDLightConversion.cpp
+++ b/Engine/Plugins/Runtime/USDCore/Source/USDUtilities/Private/USDLightConversion.cpp
@@ -46,17 +46,6 @@
 
 namespace LightConversionImpl
 {
-	/**
-	 * Calculates the solid angle in steradian that corresponds to the sphere surface area of the base of the cone
-	 * with the apex at the center of a unit sphere, and angular diameter `SourceAngleDeg`.
-	 * E.g. Sun in the sky has ~0.53 degree angular diameter -> 6.720407093551621e-05 sr
-	 * Source: https://en.wikipedia.org/wiki/Solid_angle#Cone,_spherical_cap,_hemisphere
-	 */
-	float SourceAngleToSteradian(float SourceAngleDeg)
-	{
-		return 2.0f * PI * (1.0f - FMath::Cos(FMath::DegreesToRadians(SourceAngleDeg / 2.0f)));
-	}
-
 	// Copied from USpotLightComponent::GetCosHalfConeAngle, so we don't need a component to do the same math
 	float GetSpotLightCosHalfConeAngle(float OuterConeAngle, float InnerConeAngle)
 	{
@@ -86,7 +75,7 @@ bool UsdToUnreal::ConvertLight(const pxr::UsdPrim& Prim, ULightComponentBase& Li
 
 	const bool bSRGB = true;
 	LightComponentBase.LightColor = UsdToUnreal::ConvertColor(UsdColor).ToFColor(bSRGB);
-	LightComponentBase.Intensity = UsdToUnreal::ConvertLightIntensityAttr(UsdIntensity, UsdExposure);
+	LightComponentBase.Intensity = UsdToUnreal::ConvertLightIntensity(UsdIntensity, UsdExposure);
 
 	if (ULightComponent* LightComponent = Cast<ULightComponent>(&LightComponentBase))
 	{
@@ -140,11 +129,24 @@ bool UsdToUnreal::ConvertRectLight(const pxr::UsdPrim& Prim, URectLightComponent
 	const float UsdExposure = UsdUtils::GetUsdValue<float>(RectLight.GetExposureAttr(), UsdTimeCode);
 	const float UsdWidth = UsdUtils::GetUsdValue<float>(RectLight.GetWidthAttr(), UsdTimeCode);
 	const float UsdHeight = UsdUtils::GetUsdValue<float>(RectLight.GetHeightAttr(), UsdTimeCode);
+	const double UsdNormalized = UsdUtils::GetUsdValue<bool>(RectLight.GetNormalizeAttr(), pxr::UsdTimeCode::Default());
+
+	const double NitsIntensity = UsdToUnreal::ConvertLightIntensity(UsdIntensity, UsdExposure);
 
 	LightComponent.SourceWidth = UsdToUnreal::ConvertDistance(StageInfo, UsdWidth);
 	LightComponent.SourceHeight = UsdToUnreal::ConvertDistance(StageInfo, UsdHeight);
-	LightComponent.Intensity = UsdToUnreal::ConvertLightIntensityAttr(UsdIntensity, UsdExposure);
-	LightComponent.IntensityUnits = ELightUnits::Nits;
+	if (UsdNormalized)
+	{
+		const float SolidAngle = UsdUtils::GetRectLightSolidAngle();
+		const float AreaSqMeters = UsdUtils::GetRectLightArea(LightComponent.SourceWidth / 100.0f, LightComponent.SourceHeight / 100.0f);
+		LightComponent.IntensityUnits = ELightUnits::Lumens;
+		LightComponent.Intensity = UsdUtils::ConvertIntensityFromNits(NitsIntensity, SolidAngle, AreaSqMeters, ELightUnits::Lumens);
+	}
+	else
+	{
+		LightComponent.Intensity = NitsIntensity;
+		LightComponent.IntensityUnits = ELightUnits::Nits;
+	}
 
 	return true;
 }
@@ -164,11 +166,27 @@ bool UsdToUnreal::ConvertDiskLight(const pxr::UsdPrim& Prim, URectLightComponent
 	const float UsdIntensity = UsdUtils::GetUsdValue<float>(DiskLight.GetIntensityAttr(), UsdTimeCode);
 	const float UsdExposure = UsdUtils::GetUsdValue<float>(DiskLight.GetExposureAttr(), UsdTimeCode);
 	const float UsdRadius = UsdUtils::GetUsdValue<float>(DiskLight.GetRadiusAttr(), UsdTimeCode);
+	const double UsdNormalized = UsdUtils::GetUsdValue<bool>(DiskLight.GetNormalizeAttr(), pxr::UsdTimeCode::Default());
+
+	const float RadiusUE = UsdToUnreal::ConvertDistance(StageInfo, UsdRadius);
+	const float SquareSideUE = UsdUtils::GetRectLightSideFromDiskLight(RadiusUE);
+	const double NitsIntensity = UsdToUnreal::ConvertLightIntensity(UsdIntensity, UsdExposure);
 
-	LightComponent.SourceWidth = UsdToUnreal::ConvertDistance(StageInfo, UsdRadius) * 2.f;
-	LightComponent.SourceHeight = LightComponent.SourceWidth;
-	LightComponent.Intensity = UsdToUnreal::ConvertLightIntensityAttr(UsdIntensity, UsdExposure);
-	LightComponent.IntensityUnits = ELightUnits::Nits;
+	LightComponent.SourceWidth = SquareSideUE;
+	LightComponent.SourceHeight = SquareSideUE;
+	if (UsdNormalized)
+	{
+		const float RadiusMeters = RadiusUE / 100.0f;
+		const float SolidAngle = UsdUtils::GetDiskLightSolidAngle();
+		const float AreaSqMeters = UsdUtils::GetDiskLightArea(RadiusMeters);
+		LightComponent.IntensityUnits = ELightUnits::Lumens;
+		LightComponent.Intensity = UsdUtils::ConvertIntensityFromNits(NitsIntensity, SolidAngle, AreaSqMeters, ELightUnits::Lumens);
+	}
+	else
+	{
+		LightComponent.Intensity = NitsIntensity;
+		LightComponent.IntensityUnits = ELightUnits::Nits;
+	}
 
 	return true;
 }
@@ -188,10 +206,50 @@ bool UsdToUnreal::ConvertSphereLight(const pxr::UsdPrim& Prim, UPointLightCompon
 	const float UsdIntensity = UsdUtils::GetUsdValue<float>(SphereLight.GetIntensityAttr(), UsdTimeCode);
 	const float UsdExposure = UsdUtils::GetUsdValue<float>(SphereLight.GetExposureAttr(), UsdTimeCode);
 	const float UsdRadius = UsdUtils::GetUsdValue<float>(SphereLight.GetRadiusAttr(), UsdTimeCode);
+	const double UsdNormalized = UsdUtils::GetUsdValue<bool>(SphereLight.GetNormalizeAttr(), pxr::UsdTimeCode::Default());
 
-	LightComponent.Intensity = UsdToUnreal::ConvertLightIntensityAttr(UsdIntensity, UsdExposure);
-	LightComponent.IntensityUnits = ELightUnits::Nits;
-	LightComponent.SourceRadius = UsdToUnreal::ConvertDistance(StageInfo, UsdRadius);
+	const float RadiusUE = UsdToUnreal::ConvertDistance(StageInfo, UsdRadius);
+	const double NitsIntensity = UsdToUnreal::ConvertLightIntensity(UsdIntensity, UsdExposure);
+
+	LightComponent.SourceRadius = RadiusUE;
+	if (UsdNormalized)
+	{
+		const float RadiusMeters = RadiusUE / 100.0f;
+		float SolidAngle = 1.0f;
+		float AreaSqMeters = 1.0f;
+
+		// Annoying, but historically the convention has been to call both this and ConvertLuxShapingAPI on the same prim/component
+		// as this handles the intensity part and that handles the cone angles. So here we need to redo some of the work that
+		// that function also does in order to also retrieve the angles for the unit conversion
+		if (USpotLightComponent* SpotLightComponent = Cast<USpotLightComponent>(&LightComponent))
+		{
+			pxr::UsdLuxShapingAPI ShapingAPI{Prim};
+			if (ShapingAPI)
+			{
+				const float UsdConeAngle = UsdUtils::GetUsdValue<float>(ShapingAPI.GetShapingConeAngleAttr(), UsdTimeCode);
+				const float UsdConeSoftness = UsdUtils::GetUsdValue<float>(ShapingAPI.GetShapingConeSoftnessAttr(), UsdTimeCode);
+				float InnerConeAngle = 0.0f;
+				float OuterConeAngle = 1.0f;
+				UsdToUnreal::ConvertSpotLightConeAngles(UsdConeAngle, UsdConeSoftness, OuterConeAngle, InnerConeAngle);
+
+				SolidAngle = UsdUtils::GetSpotLightSolidAngle(OuterConeAngle, InnerConeAngle);
+				AreaSqMeters = UsdUtils::GetSpotLightArea(RadiusMeters);
+			}
+		}
+		else // Just an actual point light
+		{
+			SolidAngle = UsdUtils::GetPointLightSolidAngle();
+			AreaSqMeters = UsdUtils::GetPointLightArea(RadiusMeters);
+		}
+
+		LightComponent.IntensityUnits = ELightUnits::Lumens;
+		LightComponent.Intensity = UsdUtils::ConvertIntensityFromNits(NitsIntensity, SolidAngle, AreaSqMeters, ELightUnits::Lumens);
+	}
+	else
+	{
+		LightComponent.Intensity = NitsIntensity;
+		LightComponent.IntensityUnits = ELightUnits::Nits;
+	}
 
 	return true;
 }
@@ -215,7 +273,8 @@ bool UsdToUnreal::ConvertLuxShapingAPI(const pxr::UsdPrim& Prim, USpotLightCompo
 	const float UsdConeSoftness = UsdUtils::GetUsdValue<float>(ShapingAPI.GetShapingConeSoftnessAttr(), UsdTimeCode);
 
 	float InnerConeAngle = 0.0f;
-	const float OuterConeAngle = UsdToUnreal::ConvertConeAngleSoftnessAttr(UsdConeAngle, UsdConeSoftness, InnerConeAngle);
+	float OuterConeAngle = 1.0f;
+	UsdToUnreal::ConvertSpotLightConeAngles(UsdConeAngle, UsdConeSoftness, OuterConeAngle, InnerConeAngle);
 
 	LightComponent.SetInnerConeAngle(InnerConeAngle);
 	LightComponent.SetOuterConeAngle(OuterConeAngle);
@@ -223,16 +282,31 @@ bool UsdToUnreal::ConvertLuxShapingAPI(const pxr::UsdPrim& Prim, USpotLightCompo
 	return true;
 }
 
-float UsdToUnreal::ConvertLightIntensityAttr(float UsdIntensity, float UsdExposure)
+float UsdToUnreal::ConvertLightIntensity(float UsdIntensity, float UsdExposure)
 {
 	return UsdIntensity * FMath::Exp2(UsdExposure);
 }
 
+void UsdToUnreal::ConvertSpotLightConeAngles(float InConeAngle, float InConeSoftness, float& OutOuterConeAngle, float& OutInnerConeAngle)
+{
+	OutOuterConeAngle = InConeAngle;
+	OutInnerConeAngle = InConeAngle * (1.0f - InConeSoftness);
+}
+
+PRAGMA_DISABLE_DEPRECATION_WARNINGS
+// Deprecated
+float UsdToUnreal::ConvertLightIntensityAttr(float UsdIntensity, float UsdExposure)
+{
+	return ConvertLightIntensity(UsdIntensity, UsdExposure);
+}
+
+// Deprecated
 float UsdToUnreal::ConvertDistantLightIntensityAttr(float UsdIntensity, float UsdExposure)
 {
-	return UsdToUnreal::ConvertLightIntensityAttr(UsdIntensity, UsdExposure);
+	return UsdToUnreal::ConvertLightIntensity(UsdIntensity, UsdExposure);
 }
 
+// Deprecated
 float UsdToUnreal::ConvertRectLightIntensityAttr(
 	float UsdIntensity,
 	float UsdExposure,
@@ -249,9 +323,10 @@ float UsdToUnreal::ConvertRectLightIntensityAttr(
 	// Only use PI instead of 2PI because URectLightComponent::SetLightBrightness will use just PI and not 2PI for lumen conversions, due to a cosine
 	// distribution c.f. UActorFactoryRectLight::PostSpawnActor, and the PI factor between candela and lumen for rect lights on
 	// https://docs.unrealengine.com/en-US/BuildingWorlds/LightingAndShadows/PhysicalLightUnits/index.html#point,spot,andrectlights
-	return ConvertLightIntensityAttr(UsdIntensity, UsdExposure) * PI * AreaInSqMeters;	  // Lumen = Nits * (PI sr for area light) * Area
+	return ConvertLightIntensity(UsdIntensity, UsdExposure) * PI * AreaInSqMeters;	  // Lumen = Nits * (PI sr for area light) * Area
 }
 
+// Deprecated
 float UsdToUnreal::ConvertDiskLightIntensityAttr(float UsdIntensity, float UsdExposure, float UsdRadius, const FUsdStageInfo& StageInfo)
 {
 	const float Radius = UsdToUnreal::ConvertDistance(StageInfo, UsdRadius);
@@ -261,22 +336,22 @@ float UsdToUnreal::ConvertDiskLightIntensityAttr(float UsdIntensity, float UsdEx
 	// Only use PI instead of 2PI because URectLightComponent::SetLightBrightness will use just PI and not 2PI for lumen conversions, due to a cosine
 	// distribution c.f. UActorFactoryRectLight::PostSpawnActor, and the PI factor between candela and lumen for rect lights on
 	// https://docs.unrealengine.com/en-US/BuildingWorlds/LightingAndShadows/PhysicalLightUnits/index.html#point,spot,andrectlights
-	return ConvertLightIntensityAttr(UsdIntensity, UsdExposure) * PI * AreaInSqMeters;	  // Lumen = Nits * (PI sr for area light) * Area
+	return ConvertLightIntensity(UsdIntensity, UsdExposure) * PI * AreaInSqMeters;	  // Lumen = Nits * (PI sr for area light) * Area
 }
 
+// Deprecated
 float UsdToUnreal::ConvertSphereLightIntensityAttr(float UsdIntensity, float UsdExposure, float UsdRadius, const FUsdStageInfo& StageInfo)
 {
 	float Radius = UsdToUnreal::ConvertDistance(StageInfo, UsdRadius);
 
 	float SolidAngle = 4.f * PI;
 
-	// Using solid angle for this area is possibly incorrect, but using Nits for point lights also doesn't make much sense in the first place either,
-	// but we must do it for consistency with USD
-	const float AreaInSqMeters = FMath::Max(SolidAngle * FMath::Square(Radius / 100.f), KINDA_SMALL_NUMBER);
+	const float AreaInSqMeters = FMath::Max(4.f * PI * FMath::Square(Radius / 100.f), KINDA_SMALL_NUMBER);
 
-	return ConvertLightIntensityAttr(UsdIntensity, UsdExposure) * SolidAngle * AreaInSqMeters;	  // Lumen = Nits * SolidAngle * Area
+	return ConvertLightIntensity(UsdIntensity, UsdExposure) * SolidAngle * AreaInSqMeters;	  // Lumen = Nits * SolidAngle * Area
 }
 
+// Deprecated
 float UsdToUnreal::ConvertLuxShapingAPIIntensityAttr(
 	float UsdIntensity,
 	float UsdExposure,
@@ -294,19 +369,19 @@ float UsdToUnreal::ConvertLuxShapingAPIIntensityAttr(
 	// c.f. USpotLightComponent::ComputeLightBrightness
 	float SolidAngle = 2.f * PI * (1.0f - LightConversionImpl::GetSpotLightCosHalfConeAngle(OuterConeAngle, InnerConeAngle));
 
-	// Using solid angle for this area is possibly incorrect, but using Nits for point lights also doesn't make much sense in the first place either,
-	// but we must do it for consistency with USD
-	const float AreaInSqMeters = FMath::Max(SolidAngle * FMath::Square(Radius / 100.f), KINDA_SMALL_NUMBER);
+	const float AreaInSqMeters = FMath::Max(4.f * PI * FMath::Square(Radius / 100.f), KINDA_SMALL_NUMBER);
 
-	return ConvertLightIntensityAttr(UsdIntensity, UsdExposure) * SolidAngle * AreaInSqMeters;	  // Lumen = Nits * SolidAngle * Area
+	return ConvertLightIntensity(UsdIntensity, UsdExposure) * SolidAngle * AreaInSqMeters;	  // Lumen = Nits * SolidAngle * Area
 }
 
+// Deprecated
 float UsdToUnreal::ConvertConeAngleSoftnessAttr(float UsdConeAngle, float UsdConeSoftness, float& OutInnerConeAngle)
 {
 	OutInnerConeAngle = UsdConeAngle * (1.0f - UsdConeSoftness);
 	float OuterConeAngle = UsdConeAngle;
 	return OuterConeAngle;
 }
+PRAGMA_ENABLE_DEPRECATION_WARNINGS
 
 bool UnrealToUsd::ConvertLightComponent(const ULightComponentBase& LightComponent, pxr::UsdPrim& Prim, double UsdTimeCode)
 {
@@ -339,6 +414,36 @@ bool UnrealToUsd::ConvertLightComponent(const ULightComponentBase& LightComponen
 
 	if (const ULightComponent* DerivedLightComponent = Cast<ULightComponent>(&LightComponent))
 	{
+		// If we're using a normalized intensity unit, author inputs:normalize=true, but otherwise clear the 
+		// authored value (the default is false).
+		//
+		// We only read/write the default opinions for now as we don't support animating which
+		// light unit we're using anyway.
+		ELightUnits Units = DerivedLightComponent->GetLightUnits();
+		switch (Units)
+		{
+			case ELightUnits::Candelas:
+			case ELightUnits::EV:
+			case ELightUnits::Lumens:
+			case ELightUnits::Unitless:
+			{
+				if (pxr::UsdAttribute Attr = LightAPI.CreateNormalizeAttr())
+				{
+					Attr.Set<bool>(true, UsdUtils::GetDefaultTimeCode());
+					UsdUtils::NotifyIfOverriddenOpinion(Attr);
+				}
+				break;
+			}
+			default:
+			{
+				if (pxr::UsdAttribute Attr = LightAPI.GetNormalizeAttr())
+				{
+					Attr.ClearDefault();
+				}
+				break;
+			}
+		}
+
 		if (pxr::UsdAttribute Attr = LightAPI.CreateEnableColorTemperatureAttr())
 		{
 			Attr.Set<bool>(DerivedLightComponent->bUseTemperature, UsdTimeCode);
@@ -409,9 +514,6 @@ bool UnrealToUsd::ConvertDirectionalLightComponent(const UDirectionalLightCompon
 		UsdUtils::NotifyIfOverriddenOpinion(Attr);
 	}
 
-	// USD intensity units should be in Nits == Lux / Steradian, but there is no
-	// meaningful solid angle to use to perform that conversion from Lux, so we leave intensity as-is
-
 	return true;
 }
 
@@ -432,24 +534,28 @@ bool UnrealToUsd::ConvertRectLightComponent(const URectLightComponent& LightComp
 
 	FUsdStageInfo StageInfo(Prim.GetStage());
 
+	const float WidthMeters = LightComponent.SourceWidth / 100.0f;
+	const float HeightMeters = LightComponent.SourceHeight / 100.0f;
+	const float AreaSqMeters = UsdUtils::GetRectLightArea(WidthMeters, HeightMeters);
+
+	float SolidAngle = 1.0f;
+
 	// Disk light
-	float AreaInSqMeters = 1.0f;
 	if (pxr::UsdLuxDiskLight DiskLight{Prim})
 	{
-		// Averaging and converting "diameter" to "radius"
-		const float Radius = (LightComponent.SourceWidth + LightComponent.SourceHeight) / 2.0f / 2.0f;
-		AreaInSqMeters = PI * FMath::Square(Radius / 100.0f);
+		SolidAngle = UsdUtils::GetDiskLightSolidAngle();
 
 		if (pxr::UsdAttribute Attr = DiskLight.CreateRadiusAttr())
 		{
-			Attr.Set<float>(UnrealToUsd::ConvertDistance(StageInfo, Radius), UsdTimeCode);
+			const float DiskRadiusUEUnits = UsdUtils::GetDiskLightRadiusFromRectLight(LightComponent.SourceWidth, LightComponent.SourceHeight);
+			Attr.Set<float>(UnrealToUsd::ConvertDistance(StageInfo, DiskRadiusUEUnits), UsdTimeCode);
 			UsdUtils::NotifyIfOverriddenOpinion(Attr);
 		}
 	}
 	// Rect light
 	else if (pxr::UsdLuxRectLight RectLight{Prim})
 	{
-		AreaInSqMeters = (LightComponent.SourceWidth / 100.0f) * (LightComponent.SourceHeight / 100.0f);
+		SolidAngle = UsdUtils::GetRectLightSolidAngle();
 
 		if (pxr::UsdAttribute Attr = RectLight.CreateWidthAttr())
 		{
@@ -471,19 +577,16 @@ bool UnrealToUsd::ConvertRectLightComponent(const URectLightComponent& LightComp
 	// Common for both
 	if (pxr::UsdAttribute Attr = LightAPI.CreateIntensityAttr())
 	{
-		float OldIntensity = UsdUtils::GetUsdValue<float>(Attr, UsdTimeCode);
+		float OldUsdIntensity = UsdUtils::GetUsdValue<float>(Attr, UsdTimeCode);
 
 		// Area light with no area probably shouldn't emit any light?
 		// It's not possible to set width/height less than 1 via the Details panel anyway, but just in case
-		if (FMath::IsNearlyZero(AreaInSqMeters))
+		if (FMath::IsNearlyZero(AreaSqMeters))
 		{
-			OldIntensity = 0.0f;
+			OldUsdIntensity = 0.0f;
 		}
 
-		AreaInSqMeters = FMath::Max(AreaInSqMeters, KINDA_SMALL_NUMBER);
-
-		const float Steradians = PI;
-		float FinalIntensityNits = UsdUtils::ConvertIntensityToNits(OldIntensity, Steradians, AreaInSqMeters, LightComponent.IntensityUnits);
+		const float FinalIntensityNits = UsdUtils::ConvertIntensityToNits(OldUsdIntensity, SolidAngle, AreaSqMeters, LightComponent.IntensityUnits);
 
 		Attr.Set<float>(FinalIntensityNits, UsdTimeCode);
 		UsdUtils::NotifyIfOverriddenOpinion(Attr);
@@ -509,32 +612,38 @@ bool UnrealToUsd::ConvertPointLightComponent(const UPointLightComponent& LightCo
 
 	FUsdStageInfo StageInfo(Prim.GetStage());
 
+	const float RadiusUE = LightComponent.SourceRadius;
+	const float RadiusMeters = RadiusUE / 100.0f;
+
 	if (pxr::UsdAttribute Attr = Light.CreateRadiusAttr())
 	{
-		Attr.Set<float>(UnrealToUsd::ConvertDistance(StageInfo, LightComponent.SourceRadius), UsdTimeCode);
+		Attr.Set<float>(UnrealToUsd::ConvertDistance(StageInfo, RadiusUE), UsdTimeCode);
 		UsdUtils::NotifyIfOverriddenOpinion(Attr);
 	}
 
 	if (pxr::UsdAttribute Attr = Light.CreateTreatAsPointAttr())
 	{
-		Attr.Set<bool>(FMath::IsNearlyZero(LightComponent.SourceRadius), UsdTimeCode);
+		Attr.Set<bool>(FMath::IsNearlyZero(RadiusUE), UsdTimeCode);
 		UsdUtils::NotifyIfOverriddenOpinion(Attr);
 	}
 
-	float SolidAngle = 4.f * PI;
+	float SolidAngle = 1.0f;
+	float AreaSqMeters = 1.0f;
 	if (const USpotLightComponent* SpotLightComponent = Cast<const USpotLightComponent>(&LightComponent))
 	{
-		SolidAngle = 2.f * PI * (1.0f - SpotLightComponent->GetCosHalfConeAngle());
+		SolidAngle = UsdUtils::GetSpotLightSolidAngle(SpotLightComponent->OuterConeAngle, SpotLightComponent->InnerConeAngle);
+		AreaSqMeters = UsdUtils::GetSpotLightArea(RadiusMeters);
+	}
+	else
+	{
+		SolidAngle = UsdUtils::GetPointLightSolidAngle();
+		AreaSqMeters = UsdUtils::GetPointLightArea(RadiusMeters);
 	}
 
-	// It doesn't make much physical sense to use nits for point lights in this way, but USD light intensities are always in nits
-	// so we must do something. We do the analogue on the UsdToUnreal conversion, at least. Also using the solid angle for the
-	// area calculation is possibly incorrect, but I think it depends on the chosen convention
-	const float AreaInSqMeters = FMath::Max(SolidAngle * FMath::Square(LightComponent.SourceRadius / 100.f), KINDA_SMALL_NUMBER);
 	if (pxr::UsdAttribute Attr = Light.CreateIntensityAttr())
 	{
-		const float OldIntensity = UsdUtils::GetUsdValue<float>(Attr, UsdTimeCode);
-		float FinalIntensityNits = UsdUtils::ConvertIntensityToNits(OldIntensity, SolidAngle, AreaInSqMeters, LightComponent.IntensityUnits);
+		const float OldUsdIntensity = UsdUtils::GetUsdValue<float>(Attr, UsdTimeCode);
+		const float FinalIntensityNits = UsdUtils::ConvertIntensityToNits(OldUsdIntensity, SolidAngle, AreaSqMeters, LightComponent.IntensityUnits);
 
 		Attr.Set<float>(FinalIntensityNits, UsdTimeCode);
 		UsdUtils::NotifyIfOverriddenOpinion(Attr);
@@ -611,33 +720,42 @@ bool UnrealToUsd::ConvertSpotLightComponent(const USpotLightComponent& LightComp
 		return false;
 	}
 
+	float UsdConeAngle = 1.0f;
+	float UsdConeSoftness = 1.0f;
+	UnrealToUsd::ConvertSpotLightConeAngles(LightComponent.OuterConeAngle, LightComponent.InnerConeAngle, UsdConeAngle, UsdConeSoftness);
+
 	if (pxr::UsdAttribute ConeAngleAttr = ShapingAPI.CreateShapingConeAngleAttr())
 	{
-		ConeAngleAttr.Set<float>(LightComponent.OuterConeAngle, UsdTimeCode);
+		ConeAngleAttr.Set<float>(UsdConeAngle, UsdTimeCode);
 		UsdUtils::NotifyIfOverriddenOpinion(ConeAngleAttr);
 	}
 
-	// As of March 2021 there doesn't seem to be a consensus on what softness means, according to
-	// https://groups.google.com/g/usd-interest/c/A6bc4OZjSB0 We approximate the best look here by trying to convert from inner/outer cone angle to
-	// softness according to the renderman docs
 	if (pxr::UsdAttribute SoftnessAttr = ShapingAPI.CreateShapingConeSoftnessAttr())
 	{
-		// Keep in [0, 1] range, where 1 is maximum softness, i.e. inner cone angle is zero
-		const float Softness = FMath::IsNearlyZero(LightComponent.OuterConeAngle)
-								   ? 0.0
-								   : 1.0f - LightComponent.InnerConeAngle / LightComponent.OuterConeAngle;
-		SoftnessAttr.Set<float>(Softness, UsdTimeCode);
+		SoftnessAttr.Set<float>(UsdConeSoftness, UsdTimeCode);
 		UsdUtils::NotifyIfOverriddenOpinion(SoftnessAttr);
 	}
 
 	return true;
 }
 
+void UnrealToUsd::ConvertSpotLightConeAngles(float InOuterConeAngle, float InInnerConeAngle, float& OutConeAngle, float& OutConeSoftness)
+{
+	// As of March 2021 there doesn't seem to be a consensus on what softness means, according to https://groups.google.com/g/usd-interest/c/A6bc4OZjSB0 
+	// We approximate the best look here by trying to convert from inner/outer cone angle to softness according to the renderman docs
+
+	OutConeAngle = InOuterConeAngle;
+	OutConeSoftness = FMath::IsNearlyZero(InOuterConeAngle) ? 0.0 : 1.0f - (InInnerConeAngle / InOuterConeAngle);
+}
+
+PRAGMA_DISABLE_DEPRECATION_WARNINGS
+// Deprecated
 float UnrealToUsd::ConvertLightIntensityProperty(float Intensity)
 {
 	return Intensity;
 }
 
+// Deprecated
 float UnrealToUsd::ConvertRectLightIntensityProperty(
 	float Intensity,
 	float Width,
@@ -664,6 +782,7 @@ float UnrealToUsd::ConvertRectLightIntensityProperty(
 	return UsdUtils::ConvertIntensityToNits(UsdIntensity, PI, AreaInSqMeters, SourceUnits);
 }
 
+// Deprecated
 float UnrealToUsd::ConvertRectLightIntensityProperty(float Intensity, float Radius, const FUsdStageInfo& StageInfo, ELightUnits SourceUnits)
 {
 	float UsdIntensity = UnrealToUsd::ConvertLightIntensityProperty(Intensity);
@@ -684,20 +803,19 @@ float UnrealToUsd::ConvertRectLightIntensityProperty(float Intensity, float Radi
 	return UsdUtils::ConvertIntensityToNits(UsdIntensity, PI, AreaInSqMeters, SourceUnits);
 }
 
+// Deprecated
 float UnrealToUsd::ConvertPointLightIntensityProperty(float Intensity, float SourceRadius, const FUsdStageInfo& StageInfo, ELightUnits SourceUnits)
 {
 	const float UsdIntensity = UnrealToUsd::ConvertLightIntensityProperty(Intensity);
 
 	const float SolidAngle = 4.f * PI;
 
-	// It doesn't make much physical sense to use nits for point lights in this way, but USD light intensities are always in nits
-	// so we must do something. We do the analogue on the UsdToUnreal conversion, at least. Also using the solid angle for the
-	// area calculation is possibly incorrect, but I think it depends on the chosen convention
-	const float AreaInSqMeters = FMath::Max(SolidAngle * FMath::Square(SourceRadius / 100.f), KINDA_SMALL_NUMBER);
+	const float AreaInSqMeters = FMath::Max(4.f * PI * FMath::Square(SourceRadius / 100.f), KINDA_SMALL_NUMBER);
 
 	return UsdUtils::ConvertIntensityToNits(UsdIntensity, SolidAngle, AreaInSqMeters, SourceUnits);
 }
 
+// Deprecated
 float UnrealToUsd::ConvertSpotLightIntensityProperty(
 	float Intensity,
 	float OuterConeAngle,
@@ -712,27 +830,100 @@ float UnrealToUsd::ConvertSpotLightIntensityProperty(
 	// c.f. USpotLightComponent::ComputeLightBrightness
 	const float SolidAngle = 2.f * PI * (1.0f - LightConversionImpl::GetSpotLightCosHalfConeAngle(OuterConeAngle, InnerConeAngle));
 
-	// It doesn't make much physical sense to use nits for point lights in this way, but USD light intensities are always in nits
-	// so we must do something. We do the analogue on the UsdToUnreal conversion, at least. Also using the solid angle for the
-	// area calculation is possibly incorrect, but I think it depends on the chosen convention
-	const float AreaInSqMeters = FMath::Max(SolidAngle * FMath::Square(SourceRadius / 100.f), KINDA_SMALL_NUMBER);
+	const float AreaInSqMeters = FMath::Max(4.f * PI * FMath::Square(SourceRadius / 100.f), KINDA_SMALL_NUMBER);
 
 	return UsdUtils::ConvertIntensityToNits(UsdIntensity, SolidAngle, AreaInSqMeters, SourceUnits);
 }
 
+// Deprecated
 float UnrealToUsd::ConvertOuterConeAngleProperty(float OuterConeAngle)
 {
 	return OuterConeAngle;
 }
 
+// Deprecated
 float UnrealToUsd::ConvertInnerConeAngleProperty(float InnerConeAngle, float OuterConeAngle)
 {
 	// Keep in [0, 1] range, where 1 is maximum softness, i.e. inner cone angle is zero
 	return FMath::IsNearlyZero(OuterConeAngle) ? 0.0 : 1.0f - InnerConeAngle / OuterConeAngle;
 }
+PRAGMA_ENABLE_DEPRECATION_WARNINGS
+
+float UsdUtils::GetRectLightSideFromDiskLight(float DiskLightRadius)
+{
+	// For area lights the important is the actual area, so when converting the disk lights into unreal rect lights we'll produce a
+	// square rect light that provides the same area as the USD Disk.
+	//
+	// This is needed because it ensures the engine produces the correct intensity value whenever it converts light units,
+	// as it has no idea this is a fake disk light area instead
+
+	const double DiskArea = PI * DiskLightRadius * DiskLightRadius;
+	const double SquareSide = FMath::Sqrt(DiskArea);
+	return SquareSide;
+}
+
+float UsdUtils::GetDiskLightRadiusFromRectLight(float RectLightWidth, float RectLightHeight)
+{
+	const double RectArea = RectLightWidth * RectLightHeight;
+	const float DiskRadius = FMath::Sqrt(RectArea / PI);
+	return DiskRadius;
+}
+
+float UsdUtils::GetRectLightArea(float Width, float Height)
+{
+	return Width * Height;
+}
+
+float UsdUtils::GetDiskLightArea(float Radius)
+{
+	return PI * Radius * Radius;
+}
+
+float UsdUtils::GetPointLightArea(float Radius)
+{
+	return 4.0f * PI * Radius * Radius;
+}
+
+float UsdUtils::GetSpotLightArea(float Radius)
+{
+	// The "focusing" caused by the spot light is accounted for in the solid angle, not on the area
+	return UsdUtils::GetPointLightArea(Radius);
+}
+
+float UsdUtils::GetRectLightSolidAngle()
+{
+	// For area lights sr is technically 2PI, but we cancel that with an
+	// extra factor of 2.0 here because URectLightComponent::SetLightBrightness uses just PI and not 2PI as steradian
+	// due to some cosine distribution. This also matches the PI factor between candelas and lumen for rect lights on
+	// https://docs.unrealengine.com/en-US/Engine/Rendering/LightingAndShadows/PhysicalLightUnits/index.html#point,spot,andrectlights
+	return PI;
+}
+
+float UsdUtils::GetDiskLightSolidAngle()
+{
+	// For area lights sr is technically 2PI, but we cancel that with an
+	// extra factor of 2.0 here because URectLightComponent::SetLightBrightness uses just PI and not 2PI as steradian
+	// due to some cosine distribution. This also matches the PI factor between candelas and lumen for rect lights on
+	// https://docs.unrealengine.com/en-US/Engine/Rendering/LightingAndShadows/PhysicalLightUnits/index.html#point,spot,andrectlights
+	return PI;
+}
+
+float UsdUtils::GetPointLightSolidAngle()
+{
+	return 4.0f * PI;
+}
+
+float UsdUtils::GetSpotLightSolidAngle(float OuterConeAngle, float InnerConeAngle)
+{
+	// Reference: USpotLightComponent::ComputeLightBrightness
+	return 2.0f * PI * (1.0f - LightConversionImpl::GetSpotLightCosHalfConeAngle(OuterConeAngle, InnerConeAngle));
+}
 
 float UsdUtils::ConvertIntensityToNits(float Intensity, float Steradians, float AreaInSqMeters, ELightUnits SourceUnits)
 {
+	AreaInSqMeters = FMath::Max(AreaInSqMeters, KINDA_SMALL_NUMBER);
+	Steradians = FMath::Max(Steradians, KINDA_SMALL_NUMBER);
+
 	switch (SourceUnits)
 	{
 		case ELightUnits::Candelas:
@@ -744,7 +935,9 @@ float UsdUtils::ConvertIntensityToNits(float Intensity, float Steradians, float
 			return Intensity / (Steradians * AreaInSqMeters);
 		case ELightUnits::EV:
 			// Nit = luminance (cd/m2)
-			return EV100ToLuminance(Intensity);
+			// It is not clear why dividing by the area is necessary here units-wise, but the UE components seem to do this.
+			// For example check UPointLightComponent::ComputeLightBrightness() and UPointLightComponent::SetLightBrightness()
+			return EV100ToLuminance(Intensity) / AreaInSqMeters;
 		case ELightUnits::Nits:
 			return Intensity;
 		case ELightUnits::Unitless:
@@ -758,6 +951,27 @@ float UsdUtils::ConvertIntensityToNits(float Intensity, float Steradians, float
 	return Intensity;
 }
 
+float UsdUtils::ConvertIntensityFromNits(float Intensity, float Steradians, float AreaInSqMeters, ELightUnits DestUnits)
+{
+	switch (DestUnits)
+	{
+		case ELightUnits::Candelas:
+			return Intensity * AreaInSqMeters;
+		case ELightUnits::Lumens:
+			return Intensity * (Steradians * AreaInSqMeters);
+		case ELightUnits::EV:
+			return LuminanceToEV100(Intensity * AreaInSqMeters);
+		case ELightUnits::Nits:
+			return Intensity;
+		case ELightUnits::Unitless:
+			return Intensity * (625.0f * AreaInSqMeters);
+		default:
+			break;
+	}
+
+	return Intensity;
+}
+
 #undef LOCTEXT_NAMESPACE
 
 #endif	  // #if USE_USD_SDK
diff --git a/Engine/Plugins/Runtime/USDCore/Source/USDUtilities/Private/USDPrimConversion.cpp b/Engine/Plugins/Runtime/USDCore/Source/USDUtilities/Private/USDPrimConversion.cpp
--- a/Engine/Plugins/Runtime/USDCore/Source/USDUtilities/Private/USDPrimConversion.cpp
+++ b/Engine/Plugins/Runtime/USDCore/Source/USDUtilities/Private/USDPrimConversion.cpp
@@ -1670,7 +1670,7 @@ UsdToUnreal::FPropertyTrackReader UsdToUnreal::CreatePropertyTrackReader(
 						IntensityAttr.Get<float>(&DefaultUsdIntensity);
 						ExposureAttr.Get<float>(&DefaultUsdExposure);
 
-						float Default = UsdToUnreal::ConvertLightIntensityAttr(DefaultUsdIntensity, DefaultUsdExposure);
+						float Default = UsdToUnreal::ConvertLightIntensity(DefaultUsdIntensity, DefaultUsdExposure);
 
 						Reader.FloatReader = [IntensityAttr, ExposureAttr, Default, StageInfo](double UsdTimeCode)
 						{
@@ -1680,7 +1680,7 @@ UsdToUnreal::FPropertyTrackReader UsdToUnreal::CreatePropertyTrackReader(
 							float UsdExposure = 0.0f;
 							if (IntensityAttr.Get<float>(&UsdIntensity, UsdTimeCode) && ExposureAttr.Get<float>(&UsdExposure, UsdTimeCode))
 							{
-								Result = UsdToUnreal::ConvertLightIntensityAttr(UsdIntensity, UsdExposure);
+								Result = UsdToUnreal::ConvertLightIntensity(UsdIntensity, UsdExposure);
 							}
 
 							return Result;
@@ -1718,18 +1718,20 @@ UsdToUnreal::FPropertyTrackReader UsdToUnreal::CreatePropertyTrackReader(
 						ConeAngleAttr.Get<float>(&DefaultUsdConeAngle);
 						ConeSoftnessAttr.Get<float>(&DefaultUsdConeSoftness);
 
-						float Default = 0.0f;
-						UsdToUnreal::ConvertConeAngleSoftnessAttr(DefaultUsdConeAngle, DefaultUsdConeSoftness, Default);
+						float DefaultInnerConeAngle = 0.0f;
+						float DefaultOuterConeAngle = 0.0f;
+						UsdToUnreal::ConvertSpotLightConeAngles(DefaultUsdConeAngle, DefaultUsdConeSoftness, DefaultOuterConeAngle, DefaultInnerConeAngle);
 
-						Reader.FloatReader = [ConeAngleAttr, ConeSoftnessAttr, Default](double UsdTimeCode)
+						Reader.FloatReader = [ConeAngleAttr, ConeSoftnessAttr, DefaultInnerConeAngle](double UsdTimeCode)
 						{
-							float Result = Default;
+							float Result = DefaultInnerConeAngle;
 
 							float UsdConeAngle = 90.0f;
 							float UsdConeSoftness = 0.0f;
 							if (ConeAngleAttr.Get<float>(&UsdConeAngle, UsdTimeCode) && ConeSoftnessAttr.Get<float>(&UsdConeSoftness, UsdTimeCode))
 							{
-								UsdToUnreal::ConvertConeAngleSoftnessAttr(UsdConeAngle, UsdConeSoftness, Result);
+								float OuterConeAngle = 0.0f;
+								UsdToUnreal::ConvertSpotLightConeAngles(UsdConeAngle, UsdConeSoftness, OuterConeAngle, Result);
 							}
 
 							return Result;
@@ -1756,7 +1758,7 @@ UsdToUnreal::FPropertyTrackReader UsdToUnreal::CreatePropertyTrackReader(
 						IntensityAttr.Get<float>(&DefaultUsdIntensity);
 						ExposureAttr.Get<float>(&DefaultUsdExposure);
 
-						float Default = UsdToUnreal::ConvertLightIntensityAttr(DefaultUsdIntensity, DefaultUsdExposure);
+						float Default = UsdToUnreal::ConvertLightIntensity(DefaultUsdIntensity, DefaultUsdExposure);
 
 						Reader.FloatReader = [IntensityAttr, ExposureAttr, Default, StageInfo](double UsdTimeCode)
 						{
@@ -1766,7 +1768,7 @@ UsdToUnreal::FPropertyTrackReader UsdToUnreal::CreatePropertyTrackReader(
 							float UsdExposure = 0.0f;
 							if (IntensityAttr.Get<float>(&UsdIntensity, UsdTimeCode) && ExposureAttr.Get<float>(&UsdExposure, UsdTimeCode))
 							{
-								Result = UsdToUnreal::ConvertLightIntensityAttr(UsdIntensity, UsdExposure);
+								Result = UsdToUnreal::ConvertLightIntensity(UsdIntensity, UsdExposure);
 							}
 
 							return Result;
@@ -1834,7 +1836,7 @@ UsdToUnreal::FPropertyTrackReader UsdToUnreal::CreatePropertyTrackReader(
 					IntensityAttr.Get<float>(&DefaultUsdIntensity);
 					ExposureAttr.Get<float>(&DefaultUsdExposure);
 
-					float Default = UsdToUnreal::ConvertLightIntensityAttr(DefaultUsdIntensity, DefaultUsdExposure);
+					float Default = UsdToUnreal::ConvertLightIntensity(DefaultUsdIntensity, DefaultUsdExposure);
 
 					Reader.FloatReader = [IntensityAttr, ExposureAttr, Default, StageInfo](double UsdTimeCode)
 					{
@@ -1844,7 +1846,7 @@ UsdToUnreal::FPropertyTrackReader UsdToUnreal::CreatePropertyTrackReader(
 						float UsdExposure = 0.0f;
 						if (IntensityAttr.Get<float>(&UsdIntensity, UsdTimeCode) && ExposureAttr.Get<float>(&UsdExposure, UsdTimeCode))
 						{
-							Result = UsdToUnreal::ConvertLightIntensityAttr(UsdIntensity, UsdExposure);
+							Result = UsdToUnreal::ConvertLightIntensity(UsdIntensity, UsdExposure);
 						}
 
 						return Result;
@@ -1859,18 +1861,23 @@ UsdToUnreal::FPropertyTrackReader UsdToUnreal::CreatePropertyTrackReader(
 			{
 				if (pxr::UsdAttribute Attr = DiskLight.GetRadiusAttr())
 				{
-					// Our conversion is that Width == Height == 2 * Radius
-
+					// The convention is that we fake a disk light with a rect light. For area lights the most important thing
+					// is to preserve the area though, so here we compute the side of the rect light that would produce the same
+					// total area as the disk
 					float Default = 0.0;
 					Attr.Get<float>(&Default);
-					Default = 2.0f * UsdToUnreal::ConvertDistance(StageInfo, Default);
+					float UsdArea = PI * Default * Default;
+					float UsdSquareSide = FMath::Sqrt(UsdArea);
+					Default = UsdToUnreal::ConvertDistance(StageInfo, UsdSquareSide);
 
 					Reader.FloatReader = [Attr, Default, StageInfo](double UsdTimeCode)
 					{
 						float Result = Default;
 						if (Attr.Get<float>(&Result, UsdTimeCode))
 						{
-							Result = 2.0f * UsdToUnreal::ConvertDistance(StageInfo, Result);
+							float UsdArea = PI * Result * Result;
+							float UsdSquareSide = FMath::Sqrt(UsdArea);
+							Result = UsdToUnreal::ConvertDistance(StageInfo, UsdSquareSide);
 						}
 
 						return Result;
@@ -1892,7 +1899,7 @@ UsdToUnreal::FPropertyTrackReader UsdToUnreal::CreatePropertyTrackReader(
 					IntensityAttr.Get<float>(&DefaultUsdIntensity);
 					ExposureAttr.Get<float>(&DefaultUsdExposure);
 
-					float Default = UsdToUnreal::ConvertLightIntensityAttr(DefaultUsdIntensity, DefaultUsdExposure);
+					float Default = UsdToUnreal::ConvertLightIntensity(DefaultUsdIntensity, DefaultUsdExposure);
 
 					Reader.FloatReader = [IntensityAttr, ExposureAttr, Default, StageInfo](double UsdTimeCode)
 					{
@@ -1902,7 +1909,7 @@ UsdToUnreal::FPropertyTrackReader UsdToUnreal::CreatePropertyTrackReader(
 						float UsdExposure = 0.0f;
 						if (IntensityAttr.Get<float>(&UsdIntensity, UsdTimeCode) && ExposureAttr.Get<float>(&UsdExposure, UsdTimeCode))
 						{
-							Result = UsdToUnreal::ConvertLightIntensityAttr(UsdIntensity, UsdExposure);
+							Result = UsdToUnreal::ConvertLightIntensity(UsdIntensity, UsdExposure);
 						}
 
 						return Result;
@@ -1931,8 +1938,8 @@ UsdToUnreal::FPropertyTrackReader UsdToUnreal::CreatePropertyTrackReader(
 			}
 			else if (PropertyPath == UnrealIdentifiers::IntensityPropertyName)
 			{
-				pxr::UsdAttribute IntensityAttr = SphereLight.GetIntensityAttr();
-				pxr::UsdAttribute ExposureAttr = SphereLight.GetExposureAttr();
+				pxr::UsdAttribute IntensityAttr = DistantLight.GetIntensityAttr();
+				pxr::UsdAttribute ExposureAttr = DistantLight.GetExposureAttr();
 
 				if (IntensityAttr && ExposureAttr)
 				{
@@ -1943,7 +1950,7 @@ UsdToUnreal::FPropertyTrackReader UsdToUnreal::CreatePropertyTrackReader(
 					IntensityAttr.Get<float>(&DefaultUsdIntensity);
 					ExposureAttr.Get<float>(&DefaultUsdExposure);
 
-					float Default = UsdToUnreal::ConvertLightIntensityAttr(DefaultUsdIntensity, DefaultUsdExposure);
+					float Default = UsdToUnreal::ConvertLightIntensity(DefaultUsdIntensity, DefaultUsdExposure);
 
 					Reader.FloatReader = [IntensityAttr, ExposureAttr, Default, StageInfo](double UsdTimeCode)
 					{
@@ -1953,7 +1960,7 @@ UsdToUnreal::FPropertyTrackReader UsdToUnreal::CreatePropertyTrackReader(
 						float UsdExposure = 0.0f;
 						if (IntensityAttr.Get<float>(&UsdIntensity, UsdTimeCode) && ExposureAttr.Get<float>(&UsdExposure, UsdTimeCode))
 						{
-							Result = UsdToUnreal::ConvertLightIntensityAttr(UsdIntensity, UsdExposure);
+							Result = UsdToUnreal::ConvertLightIntensity(UsdIntensity, UsdExposure);
 						}
 
 						return Result;
@@ -1980,7 +1987,7 @@ UsdToUnreal::FPropertyTrackReader UsdToUnreal::CreatePropertyTrackReader(
 					IntensityAttr.Get<float>(&DefaultUsdIntensity);
 					ExposureAttr.Get<float>(&DefaultUsdExposure);
 
-					float Default = UsdToUnreal::ConvertLightIntensityAttr(DefaultUsdIntensity, DefaultUsdExposure);
+					float Default = UsdToUnreal::ConvertLightIntensity(DefaultUsdIntensity, DefaultUsdExposure);
 
 					Reader.FloatReader = [IntensityAttr, ExposureAttr, Default, StageInfo](double UsdTimeCode)
 					{
@@ -1990,7 +1997,7 @@ UsdToUnreal::FPropertyTrackReader UsdToUnreal::CreatePropertyTrackReader(
 						float UsdExposure = 0.0f;
 						if (IntensityAttr.Get<float>(&UsdIntensity, UsdTimeCode) && ExposureAttr.Get<float>(&UsdExposure, UsdTimeCode))
 						{
-							Result = UsdToUnreal::ConvertLightIntensityAttr(UsdIntensity, UsdExposure);
+							Result = UsdToUnreal::ConvertLightIntensity(UsdIntensity, UsdExposure);
 						}
 
 						return Result;
@@ -5929,27 +5936,22 @@ UnrealToUsd::FPropertyTrackWriter UnrealToUsd::CreatePropertyTrackWriter(
 						if (Attr && RadiusAttr && ConeAngleAttr && ConeSoftnessAttr)
 						{
 							Result.FloatWriter =
-								[Attr, RadiusAttr, ConeAngleAttr, ConeSoftnessAttr, StageInfo, Units](float UEValue, double UsdTimeCode)
+								[Attr, RadiusAttr, ConeAngleAttr, ConeSoftnessAttr, StageInfo, Units](float UEIntensity, double UsdTimeCode)
 							{
 								const float UsdConeAngle = UsdUtils::GetUsdValue<float>(ConeAngleAttr, UsdTimeCode);
 								const float UsdConeSoftness = UsdUtils::GetUsdValue<float>(ConeSoftnessAttr, UsdTimeCode);
 								const float UsdRadius = UsdUtils::GetUsdValue<float>(RadiusAttr, UsdTimeCode);
 
 								float InnerConeAngle = 0.0f;
-								const float OuterConeAngle = UsdToUnreal::ConvertConeAngleSoftnessAttr(UsdConeAngle, UsdConeSoftness, InnerConeAngle);
-								const float SourceRadius = UsdToUnreal::ConvertDistance(StageInfo, UsdRadius);
-
-								Attr.Set(
-									UnrealToUsd::ConvertSpotLightIntensityProperty(
-										UEValue,
-										OuterConeAngle,
-										InnerConeAngle,
-										SourceRadius,
-										StageInfo,
-										Units
-									),
-									UsdTimeCode
-								);
+								float OuterConeAngle = 1.0f;
+								UsdToUnreal::ConvertSpotLightConeAngles(UsdConeAngle, UsdConeSoftness, OuterConeAngle, InnerConeAngle);
+
+								const float RadiusMeters = UsdToUnreal::ConvertDistance(StageInfo, UsdRadius) / 100.0f;
+								const float AreaSqMeters = UsdUtils::GetSpotLightArea(RadiusMeters);
+								const float SolidAngle = UsdUtils::GetSpotLightSolidAngle(OuterConeAngle, InnerConeAngle);
+								const float UsdIntensity = UsdUtils::ConvertIntensityToNits(UEIntensity, SolidAngle, AreaSqMeters, Units);
+
+								Attr.Set(UsdIntensity, UsdTimeCode);
 							};
 						}
 					}
@@ -5974,13 +5976,18 @@ UnrealToUsd::FPropertyTrackWriter UnrealToUsd::CreatePropertyTrackWriter(
 
 						if (ConeAngleAttr && Attr)
 						{
-							Result.FloatWriter = [Attr, ConeAngleAttr](float UEValue, double UsdTimeCode)
+							Result.FloatWriter = [Attr, ConeAngleAttr](float UEInnerConeAngle, double UsdTimeCode)
 							{
-								const float UsdConeAngle = UsdUtils::GetUsdValue<float>(ConeAngleAttr, UsdTimeCode);
-								const float OuterConeAngle = UsdConeAngle;
+								// Awkwardly, to write the cone softness we need not only the inner cone angle, but
+								// the outer cone angle too. We know that the outer cone angle matches the USD cone angle,
+								// so here we just read the angle from USD
+								float UsdConeAngle = UsdUtils::GetUsdValue<float>(ConeAngleAttr, UsdTimeCode);
+								float OuterConeAngle = UsdConeAngle;
+
+								float UsdConeSoftness = 0.0f;
+								UnrealToUsd::ConvertSpotLightConeAngles(OuterConeAngle, UEInnerConeAngle, UsdConeAngle, UsdConeSoftness);
 
-								const float OutNewSoftness = UnrealToUsd::ConvertInnerConeAngleProperty(UEValue, OuterConeAngle);
-								Attr.Set(OutNewSoftness, UsdTimeCode);
+								Attr.Set(UsdConeSoftness, UsdTimeCode);
 							};
 						}
 					}
@@ -6009,13 +6016,16 @@ UnrealToUsd::FPropertyTrackWriter UnrealToUsd::CreatePropertyTrackWriter(
 
 						if (Attr && RadiusAttr)
 						{
-							Result.FloatWriter = [Attr, RadiusAttr, StageInfo, Units](float UEValue, double UsdTimeCode)
+							Result.FloatWriter = [Attr, RadiusAttr, StageInfo, Units](float UEIntensity, double UsdTimeCode)
 							{
-								const float SourceRadius = UsdToUnreal::ConvertDistance(
-									StageInfo,
-									UsdUtils::GetUsdValue<float>(RadiusAttr, UsdTimeCode)
-								);
-								Attr.Set(UnrealToUsd::ConvertPointLightIntensityProperty(UEValue, SourceRadius, StageInfo, Units), UsdTimeCode);
+								const float UsdRadius = UsdUtils::GetUsdValue<float>(RadiusAttr, UsdTimeCode);
+
+								const float RadiusMeters = UsdToUnreal::ConvertDistance(StageInfo, UsdRadius) / 100.0f;
+								const float SolidAngle = UsdUtils::GetPointLightSolidAngle();
+								const float AreaSqMeters = UsdUtils::GetPointLightArea(RadiusMeters);
+								const float UsdIntensity = UsdUtils::ConvertIntensityToNits(UEIntensity, SolidAngle, AreaSqMeters, Units);
+
+								Attr.Set(UsdIntensity, UsdTimeCode);
 							};
 						}
 					}
@@ -6071,12 +6081,18 @@ UnrealToUsd::FPropertyTrackWriter UnrealToUsd::CreatePropertyTrackWriter(
 
 					if (Attr && WidthAttr && HeightAttr)
 					{
-						Result.FloatWriter = [Attr, WidthAttr, HeightAttr, StageInfo, Units](float UEValue, double UsdTimeCode)
+						Result.FloatWriter = [Attr, WidthAttr, HeightAttr, StageInfo, Units](float UEIntensity, double UsdTimeCode)
 						{
-							const float Width = UsdToUnreal::ConvertDistance(StageInfo, UsdUtils::GetUsdValue<float>(WidthAttr, UsdTimeCode));
-							const float Height = UsdToUnreal::ConvertDistance(StageInfo, UsdUtils::GetUsdValue<float>(HeightAttr, UsdTimeCode));
+							const float UsdWidth = UsdUtils::GetUsdValue<float>(WidthAttr, UsdTimeCode);
+							const float UsdHeight = UsdUtils::GetUsdValue<float>(HeightAttr, UsdTimeCode);
+
+							const float WidthMeters = UsdToUnreal::ConvertDistance(StageInfo, UsdWidth) / 100.0f;
+							const float HeightMeters = UsdToUnreal::ConvertDistance(StageInfo, UsdHeight) / 100.0f;
+							const float SolidAngle = UsdUtils::GetRectLightSolidAngle();
+							const float AreaSqMeters = UsdUtils::GetRectLightArea(WidthMeters, HeightMeters);
+							const float UsdIntensity = UsdUtils::ConvertIntensityToNits(UEIntensity, SolidAngle, AreaSqMeters, Units);
 
-							Attr.Set(UnrealToUsd::ConvertRectLightIntensityProperty(UEValue, Width, Height, StageInfo, Units), UsdTimeCode);
+							Attr.Set(UsdIntensity, UsdTimeCode);
 						};
 					}
 				}
@@ -6098,18 +6114,20 @@ UnrealToUsd::FPropertyTrackWriter UnrealToUsd::CreatePropertyTrackWriter(
 
 						Result.FloatWriter = [Attr, StageInfo](float UEValue, double UsdTimeCode)
 						{
-							Attr.Set(UnrealToUsd::ConvertDistance(StageInfo, UEValue * 0.5f), UsdTimeCode);
+							float UEArea = UEValue * UEValue;
+							float UERadius = FMath::Sqrt(UEArea / PI);
+							Attr.Set(UnrealToUsd::ConvertDistance(StageInfo, UERadius), UsdTimeCode);
 						};
 					}
 				}
 				else if (PropertyPath == UnrealIdentifiers::IntensityPropertyName)
 				{
-					Attr = RectLight.GetIntensityAttr();
+					Attr = DiskLight.GetIntensityAttr();
 					pxr::UsdAttribute RadiusAttr = DiskLight.GetRadiusAttr();
 
 					// Always clear exposure because we'll put all of our "light intensity" on the intensity attr and assume exposure
 					// is zero, as we can't manipulate something like that exposure directly from UE anyway
-					if (pxr::UsdAttribute ExposureAttr = RectLight.GetExposureAttr())
+					if (pxr::UsdAttribute ExposureAttr = DiskLight.GetExposureAttr())
 					{
 						ExposureAttr.Clear();
 					}
@@ -6123,11 +6141,16 @@ UnrealToUsd::FPropertyTrackWriter UnrealToUsd::CreatePropertyTrackWriter(
 
 					if (Attr && RadiusAttr)
 					{
-						Result.FloatWriter = [Attr, RadiusAttr, StageInfo, Units](float UEValue, double UsdTimeCode)
+						Result.FloatWriter = [Attr, RadiusAttr, StageInfo, Units](float UEIntensity, double UsdTimeCode)
 						{
-							const float Radius = UsdToUnreal::ConvertDistance(StageInfo, UsdUtils::GetUsdValue<float>(RadiusAttr, UsdTimeCode));
+							const float UsdRadius = UsdUtils::GetUsdValue<float>(RadiusAttr, UsdTimeCode);
 
-							Attr.Set(UnrealToUsd::ConvertRectLightIntensityProperty(UEValue, Radius, StageInfo, Units), UsdTimeCode);
+							const float RadiusMeters = UsdToUnreal::ConvertDistance(StageInfo, UsdRadius) / 100.0f;
+							const float SolidAngle = UsdUtils::GetDiskLightSolidAngle();
+							const float AreaSqMeters = UsdUtils::GetDiskLightArea(RadiusMeters);
+							const float UsdIntensity = UsdUtils::ConvertIntensityToNits(UEIntensity, SolidAngle, AreaSqMeters, Units);
+
+							Attr.Set(UsdIntensity, UsdTimeCode);
 						};
 					}
 				}
@@ -6151,16 +6174,17 @@ UnrealToUsd::FPropertyTrackWriter UnrealToUsd::CreatePropertyTrackWriter(
 
 					// Always clear exposure because we'll put all of our "light intensity" on the intensity attr and assume exposure
 					// is zero, as we can't manipulate something like that exposure directly from UE anyway
-					if (pxr::UsdAttribute ExposureAttr = RectLight.GetExposureAttr())
+					if (pxr::UsdAttribute ExposureAttr = DistantLight.GetExposureAttr())
 					{
 						ExposureAttr.Clear();
 					}
 
 					if (Attr)
 					{
-						Result.FloatWriter = [Attr](float UEValue, double UsdTimeCode)
+						Result.FloatWriter = [Attr](float UEIntensity, double UsdTimeCode)
 						{
-							Attr.Set(UnrealToUsd::ConvertLightIntensityProperty(UEValue), UsdTimeCode);
+							// Both UE and USD exclusively use lux for directional light intensity
+							Attr.Set(UEIntensity, UsdTimeCode);
 						};
 					}
 				}
@@ -6180,9 +6204,10 @@ UnrealToUsd::FPropertyTrackWriter UnrealToUsd::CreatePropertyTrackWriter(
 
 					if (Attr)
 					{
-						Result.FloatWriter = [Attr](float UEValue, double UsdTimeCode)
+						Result.FloatWriter = [Attr](float UEIntensity, double UsdTimeCode)
 						{
-							Attr.Set(UnrealToUsd::ConvertLightIntensityProperty(UEValue), UsdTimeCode);
+							// Both UE and USD use a flat multiplier for dome/skylight intensity
+							Attr.Set(UEIntensity, UsdTimeCode);
 						};
 					}
 				}
diff --git a/Engine/Plugins/Runtime/USDCore/Source/USDUtilities/Public/USDLightConversion.h b/Engine/Plugins/Runtime/USDCore/Source/USDUtilities/Public/USDLightConversion.h
--- a/Engine/Plugins/Runtime/USDCore/Source/USDUtilities/Public/USDLightConversion.h
+++ b/Engine/Plugins/Runtime/USDCore/Source/USDUtilities/Public/USDLightConversion.h
@@ -88,27 +88,32 @@ namespace UsdToUnreal
 		double UsdTimeCode = UsdUtils::GetDefaultTimeCode()
 	);
 
-	// These are separately exposed so that they can be reused when reading data into MovieScene tracks.
-	// The other attribute conversions are mostly trivial like a single call to UsdToUnreal::ConvertDistance.
+	/** Applies the USD exposure, essentially just returning UsdIntensity * 2.0^(UsdExposure) */
+	USDUTILITIES_API float ConvertLightIntensity(float UsdIntensity, float UsdExposure);
+
+	/**
+	 * Receives InConeAngle (degrees) and InConeSoftness ([0,1] range) USD spot light property values, and produces
+	 * the corresponding Unreal outer and inner cone angles (also degrees) for a spot light component.
+	 *
+	 * Softness 0 corresponds to a hard transition where the outer and inner angle the same: Max brightness from center until the cone angle, and zero after it
+	 * Softness 0.5 corresponds to an inner angle with half the size of the outer angle: Max brightness until the inner angle, and it lerps to zero until the outer angle
+	 * Softness 1 corresponds to zero inner angle, so the brightness lerps from max at the center all the way to zero on the outer cone angle
+	 */
+	USDUTILITIES_API void ConvertSpotLightConeAngles(float InConeAngle, float InConeSoftness, float& OutOuterConeAngle, float& OutInnerConeAngle);
+
+	UE_DEPRECATED(5.8, "This has been renamed to ConvertLightIntensity.")
 	USDUTILITIES_API float ConvertLightIntensityAttr(float UsdIntensity, float UsdExposure);
+	UE_DEPRECATED(5.8, "No longer used as this implicitly hard-coded the intensity units. Compute the solid angle and area with the other utils functions and use ConvertIntensityFromNits.")
 	USDUTILITIES_API float ConvertDistantLightIntensityAttr(float UsdIntensity, float UsdExposure);
-	USDUTILITIES_API float ConvertRectLightIntensityAttr(
-		float UsdIntensity,
-		float UsdExposure,
-		float UsdWidth,
-		float UsdHeight,
-		const FUsdStageInfo& StageInfo
-	);
+	UE_DEPRECATED(5.8, "No longer used as this implicitly hard-coded the intensity units. Compute the solid angle and area with the other utils functions and use ConvertIntensityFromNits.")
+	USDUTILITIES_API float ConvertRectLightIntensityAttr(float UsdIntensity, float UsdExposure, float UsdWidth, float UsdHeight, const FUsdStageInfo& StageInfo);
+	UE_DEPRECATED(5.8, "No longer used as this implicitly hard-coded the intensity units. Compute the solid angle and area with the other utils functions and use ConvertIntensityFromNits.")
 	USDUTILITIES_API float ConvertDiskLightIntensityAttr(float UsdIntensity, float UsdExposure, float UsdRadius, const FUsdStageInfo& StageInfo);
+	UE_DEPRECATED(5.8, "No longer used as this implicitly hard-coded the intensity units. Compute the solid angle and area with the other utils functions and use ConvertIntensityFromNits.")
 	USDUTILITIES_API float ConvertSphereLightIntensityAttr(float UsdIntensity, float UsdExposure, float UsdRadius, const FUsdStageInfo& StageInfo);
-	USDUTILITIES_API float ConvertLuxShapingAPIIntensityAttr(
-		float UsdIntensity,
-		float UsdExposure,
-		float UsdRadius,
-		float UsdConeAngle,
-		float UsdConeSoftness,
-		const FUsdStageInfo& StageInfo
-	);
+	UE_DEPRECATED(5.8, "No longer used as this implicitly hard-coded the intensity units. Compute the solid angle and area with the other utils functions and use ConvertIntensityFromNits.")
+	USDUTILITIES_API float ConvertLuxShapingAPIIntensityAttr(float UsdIntensity, float UsdExposure, float UsdRadius, float UsdConeAngle, float UsdConeSoftness, const FUsdStageInfo& StageInfo);
+	UE_DEPRECATED(5.8, "Use ConvertSpotLightConeAngles().")
 	USDUTILITIES_API float ConvertConeAngleSoftnessAttr(float UsdConeAngle, float UsdConeSoftness, float& OutInnerConeAngle);
 }
 
@@ -145,41 +150,62 @@ namespace UnrealToUsd
 		double UsdTimeCode = UsdUtils::GetDefaultTimeCode()
 	);
 
+	/**
+	 * Produces the USD cone angle (degrees) and softness ([0,1] range) values for a spot light that corresponds to an unreal outer and inner cone angles (also degrees).
+	 * This is the inverse of UsdToUnreal::ConvertSpotLightConeAngles().
+	 */
+	USDUTILITIES_API void ConvertSpotLightConeAngles(float InOuterConeAngle, float InInnerConeAngle, float& OutConeAngle, float& OutConeSoftness);
+
+	UE_DEPRECATED(5.8, "No longer used as it didn't really do anything.")
 	USDUTILITIES_API float ConvertLightIntensityProperty(float Intensity);
-	USDUTILITIES_API float ConvertRectLightIntensityProperty(
-		float Intensity,
-		float Width,
-		float Height,
-		const FUsdStageInfo& StageInfo,
-		ELightUnits SourceUnits
-	);
-	USDUTILITIES_API float ConvertRectLightIntensityProperty(
-		float Intensity,
-		float Radius,
-		const FUsdStageInfo& StageInfo,
-		ELightUnits SourceUnits
-	);	  // For RectLight -> UsdLuxDiskLight
-	USDUTILITIES_API float ConvertPointLightIntensityProperty(
-		float Intensity,
-		float SourceRadius,
-		const FUsdStageInfo& StageInfo,
-		ELightUnits SourceUnits
-	);
-	USDUTILITIES_API float ConvertSpotLightIntensityProperty(
-		float Intensity,
-		float OuterConeAngle,
-		float InnerConeAngle,
-		float SourceRadius,
-		const FUsdStageInfo& StageInfo,
-		ELightUnits SourceUnits
-	);
-	USDUTILITIES_API float ConvertOuterConeAngleProperty(float OuterConeAngle);							 // Returns ConeAngleAttr() in USD
-	USDUTILITIES_API float ConvertInnerConeAngleProperty(float InnerConeAngle, float OuterConeAngle);	 // Returns ConeSoftnessAttr() in USD
+	UE_DEPRECATED(5.8, "Compute the solid angle and area with the other utils functions and use ConvertIntensityToNits().")
+	USDUTILITIES_API float ConvertRectLightIntensityProperty(float Intensity, float Width, float Height, const FUsdStageInfo& StageInfo, ELightUnits SourceUnits);
+	UE_DEPRECATED(5.8, "Compute the solid angle and area with the other utils functions and use ConvertIntensityToNits().")
+	USDUTILITIES_API float ConvertRectLightIntensityProperty(float Intensity, float Radius, const FUsdStageInfo& StageInfo, ELightUnits SourceUnits);
+	UE_DEPRECATED(5.8, "Compute the solid angle and area with the other utils functions and use ConvertIntensityToNits().")
+	USDUTILITIES_API float ConvertPointLightIntensityProperty(float Intensity, float SourceRadius, const FUsdStageInfo& StageInfo, ELightUnits SourceUnits);
+	UE_DEPRECATED(5.8, "Compute the solid angle and area with the other utils functions and use ConvertIntensityToNits().")
+	USDUTILITIES_API float ConvertSpotLightIntensityProperty(float Intensity, float OuterConeAngle, float InnerConeAngle, float SourceRadius, const FUsdStageInfo& StageInfo, ELightUnits SourceUnits);
+	UE_DEPRECATED(5.8, "Use ConvertSpotLightConeAngles().")
+	USDUTILITIES_API float ConvertOuterConeAngleProperty(float OuterConeAngle);
+	UE_DEPRECATED(5.8, "Use ConvertSpotLightConeAngles().")
+	USDUTILITIES_API float ConvertInnerConeAngleProperty(float InnerConeAngle, float OuterConeAngle);
 }
 
 namespace UsdUtils
 {
-	USDUTILITIES_API float ConvertIntensityToNits(float Intensity, float Steradians, float AreaInSqMeters, ELightUnits SourceUnits);
+	/**
+	 * Returns the side length of a square that possesses the same area as a circle with DiskLightRadius.
+	 * Output units are the same as the input units.
+	 */
+	 USDUTILITIES_API float GetRectLightSideFromDiskLight(float DiskLightRadius);
+
+	/**
+	 * Returns the radius of a circle that posseses the same area as a rectangle with sides RectLightWidth and RectLightHeight 
+	 * Output units are the same as the input units.
+	 */
+	USDUTILITIES_API float GetDiskLightRadiusFromRectLight(float RectLightWidth, float RectLightHeight);
+
+	/**
+	 * Returns the values that should be used as "surface area" for unreal lights of different types.
+	 * Output units are the square of the input units.
+	 */
+	USDUTILITIES_API float GetRectLightArea(float Width, float Height);
+	USDUTILITIES_API float GetDiskLightArea(float Radius);
+	USDUTILITIES_API float GetPointLightArea(float Radius);
+	USDUTILITIES_API float GetSpotLightArea(float Radius);
+
+	/** Returns the values that should be used as "solid angle" for unreal lights of different types, in steradians. */
+	USDUTILITIES_API float GetRectLightSolidAngle();
+	USDUTILITIES_API float GetDiskLightSolidAngle();
+	USDUTILITIES_API float GetPointLightSolidAngle();
+	USDUTILITIES_API float GetSpotLightSolidAngle(float OuterConeAngle, float InnerConeAngle);
+
+	/** For a light with the provided SolidAngle and surface AreaInSqMeters, converts the Intensity value in the provided SourceUnits into Intensity in Nits */
+	USDUTILITIES_API float ConvertIntensityToNits(float Intensity, float SolidAngle, float AreaInSqMeters, ELightUnits SourceUnits);
+
+	/** For a light with the provided SolidAngle and surface AreaInSqMeters, converts the Intensity value in Nits into the provided DestUnits */
+	USDUTILITIES_API float ConvertIntensityFromNits(float Intensity, float SolidAngle, float AreaInSqMeters, ELightUnits DestUnits);
 }
 
 #endif	  // #if USE_USD_SDK
