USD: Updated groom import with support for missing attributes

#Synced-CL 51095459

diff --git a/Engine/Plugins/Runtime/USDCore/Source/USDUtilities/Private/USDGroomConversion.cpp b/Engine/Plugins/Runtime/USDCore/Source/USDUtilities/Private/USDGroomConversion.cpp
--- a/Engine/Plugins/Runtime/USDCore/Source/USDUtilities/Private/USDGroomConversion.cpp
+++ b/Engine/Plugins/Runtime/USDCore/Source/USDUtilities/Private/USDGroomConversion.cpp
@@ -106,8 +106,9 @@ namespace UE::UsdGroomConversion::Private
 		FString PrimvarName = UsdToUnreal::ConvertToken(Primvar.GetPrimvarName());
 		FName AttributeName(*PrimvarName);
 
+		// ComputeFlattened can handle indexed primvar; for non-indexed primvar, this is equivalent to Get
 		pxr::VtArray<USDType> Values;
-		Primvar.Get(&Values, TimeCode);
+		Primvar.ComputeFlattened(&Values, TimeCode);
 		int32 NumValues = Values.size();
 
 		if (NumValues == 0)
@@ -178,8 +179,6 @@ namespace UE::UsdGroomConversion::Private
 				continue;
 			}
 
-			// #ueent_todo: Check if the groom_color attribute is animated (for GroomCache)
-
 			// Dispatch the USD to UE conversion by the primvar scalar type
 			pxr::SdfValueTypeName ScalarTypeName = Primvar.GetTypeName().GetScalarType();
 			if (ScalarTypeName == pxr::SdfValueTypeNames->Int)
@@ -209,6 +208,68 @@ namespace UE::UsdGroomConversion::Private
 		}
 	}
 
+	/** Convert the given USD attribute to groom-scope hair attribute (constant value that applies to the whole groom) */
+	template<typename USDType, typename UEType>
+	void ConvertGroomAttribute(const pxr::UsdAttribute& Attribute, const pxr::UsdTimeCode& TimeCode, FHairDescription& HairDescription)
+	{
+		FString PropertyName = UsdToUnreal::ConvertToken(Attribute.GetBaseName());
+		FName AttributeName(*PropertyName);
+
+		USDType Value;
+		if (!Attribute.Get(&Value, TimeCode))
+		{
+			return;
+		}
+
+		TGroomAttributesRef<UEType> GroomAttributeRef = HairDescription.GroomAttributes().GetAttributesRef<UEType>(AttributeName);
+		if (!GroomAttributeRef.IsValid())
+		{
+			HairDescription.GroomAttributes().RegisterAttribute<UEType>(AttributeName, 1, TConverter<UEType>::DefaultValue());
+			GroomAttributeRef = HairDescription.GroomAttributes().GetAttributesRef<UEType>(AttributeName);
+		}
+
+		GroomAttributeRef[0] = TConverter<UEType>::ConvertType(Value);
+	}
+
+	/** Convert the given attributes to groom-scope hair attributes (constant value that applies to the whole groom) */
+	void ConvertGroomAttributes(const std::vector<pxr::UsdProperty>& Properties, const pxr::UsdTimeCode& TimeCode, FHairDescription& HairDescription)
+	{
+		for (const pxr::UsdProperty& Property : Properties)
+		{
+			pxr::UsdAttribute Attribute = Property.As<pxr::UsdAttribute>();
+			if (!Attribute || !Attribute.HasValue())
+			{
+				continue;
+			}
+
+			// Process only groom attributes
+			FString PropertyName = UsdToUnreal::ConvertToken(Property.GetBaseName());
+			if (!IsAttributeValid(PropertyName))
+			{
+				continue;
+			}
+
+			// Dispatch the USD to UE conversion by the attribute scalar type
+			pxr::SdfValueTypeName ScalarTypeName = Attribute.GetTypeName().GetScalarType();
+			if (ScalarTypeName == pxr::SdfValueTypeNames->Int)
+			{
+				ConvertGroomAttribute<int, int>(Attribute, TimeCode, HairDescription);
+			}
+			else if (ScalarTypeName == pxr::SdfValueTypeNames->Float)
+			{
+				ConvertGroomAttribute<float, float>(Attribute, TimeCode, HairDescription);
+			}
+			else if (ScalarTypeName == pxr::SdfValueTypeNames->Vector3f || ScalarTypeName == pxr::SdfValueTypeNames->Float3)
+			{
+				ConvertGroomAttribute<pxr::GfVec3f, FVector3f>(Attribute, TimeCode, HairDescription);
+			}
+			else if (ScalarTypeName == pxr::SdfValueTypeNames->String)
+			{
+				ConvertGroomAttribute<std::string, FName>(Attribute, TimeCode, HairDescription);
+			}
+		}
+	}
+
 	/** Fill the AnimInfo with the time data for the given attribute, if it's animated */
 	void UpdateAnimInfo(const pxr::UsdAttribute& Attr, EGroomCacheAttributes Flag, FGroomAnimationInfo* AnimInfo)
 	{
@@ -236,7 +297,8 @@ namespace UsdToUnreal
 		const pxr::UsdTimeCode& TimeCode,
 		const FTransform& ParentTransform,
 		FHairDescription& HairDescription,
-		FGroomAnimationInfo* AnimInfo
+		FGroomAnimationInfo* AnimInfo,
+		bool bCheckGroomAttributes
 	)
 	{
 		TRACE_CPUPROFILER_EVENT_SCOPE(UsdToUnreal::ConvertGroomHierarchy);
@@ -249,9 +311,6 @@ namespace UsdToUnreal
 		FTransform TransformToPropagate = ParentTransform;
 		if (pxr::UsdGeomBasisCurves Curves = pxr::UsdGeomBasisCurves(Prim))
 		{
-			// #ueent_todo: Add support for GroomCache import (needs AnimInfo)
-			// There are also curve attributes that are not currently queried (knots, orders, basis, type, etc.)
-
 			// Get curve and curve vextex count and vertex positions, the minimal data set required
 			const int32 NumCurves = Curves.GetCurveCount(TimeCode);
 			if (NumCurves == 0)
@@ -324,9 +383,8 @@ namespace UsdToUnreal
 				const float Scale = !FMath::IsNearlyEqual(StageInfo.MetersPerUnit, UEMetersPerUnit)
 										? (float)(StageInfo.MetersPerUnit / UEMetersPerUnit)
 										: 1.0f;
-				if (Interpolation == pxr::UsdGeomTokens->uniform || Interpolation == pxr::UsdGeomTokens->constant)
+				if (Interpolation == pxr::UsdGeomTokens->uniform)
 				{
-					const float ConstWidth = WidthsArray[0] * Scale;
 					TStrandAttributesRef<float> WidthStrandAttributeRef = HairDescription.StrandAttributes().GetAttributesRef<float>(
 						HairAttribute::Strand::Width
 					);
@@ -338,9 +396,7 @@ namespace UsdToUnreal
 
 					for (int32 Index = 0; Index < NumCurves; ++Index)
 					{
-						WidthStrandAttributeRef[FStrandID(StartStrandID + Index)] = Interpolation == pxr::UsdGeomTokens->constant
-																						? ConstWidth
-																						: WidthsArray[Index] * Scale;
+						WidthStrandAttributeRef[FStrandID(StartStrandID + Index)] = WidthsArray[Index] * Scale;
 					}
 				}
 				else if (Interpolation == pxr::UsdGeomTokens->vertex)
@@ -359,11 +415,23 @@ namespace UsdToUnreal
 						VertexAttributeRef[FVertexID(StartVertexID + Index)] = WidthsArray[Index] * Scale;
 					}
 				}
+				else if (bCheckGroomAttributes && Interpolation == pxr::UsdGeomTokens->constant)
+				{
+					// This is a more compact way to store a constant width
+					TGroomAttributesRef<float> GroomAttributeRef = HairDescription.GroomAttributes().GetAttributesRef<float>(HairAttribute::Groom::Width);
+					if (!GroomAttributeRef.IsValid())
+					{
+						HairDescription.GroomAttributes().RegisterAttribute<float>(HairAttribute::Groom::Width);
+						GroomAttributeRef = HairDescription.GroomAttributes().GetAttributesRef<float>(HairAttribute::Groom::Width);
+					}
+
+					GroomAttributeRef[0] = WidthsArray[0] * Scale;
+				}
 			}
 
 			// Extract the groom attributes (at vertex and strand scope) from the generic primvars
 			pxr::UsdGeomPrimvarsAPI PrimvarsAPI(Prim);
-			std::vector<pxr::UsdGeomPrimvar> Primvars = PrimvarsAPI.GetPrimvars();
+			std::vector<pxr::UsdGeomPrimvar> Primvars = PrimvarsAPI.GetPrimvarsWithValues();
 			if (!Primvars.empty())
 			{
 				UE::UsdGroomConversion::Private::ConvertPrimvars(
@@ -375,9 +443,16 @@ namespace UsdToUnreal
 					NumPoints,
 					HairDescription
 				);
-			}
 
-			// Groom-scope attributes are not currently extracted since they are not strictly necessary
+				// Check if color attribute is animated for GroomCache
+				// Other animated primvars/attributes are not supported by the GroomCache
+				static const pxr::TfToken GroomColorToken(UnrealToUsd::ConvertName(HairAttribute::Vertex::Color).Get());
+				pxr::UsdGeomPrimvar ColorPrimvar = PrimvarsAPI.GetPrimvar(GroomColorToken);
+				if (ColorPrimvar)
+				{
+					UE::UsdGroomConversion::Private::UpdateAnimInfo(ColorPrimvar, EGroomCacheAttributes::Color, AnimInfo);
+				}
+			}
 
 			// Following the USD recommendation that gprims be not nested, a Curves prim is handled as a leaf
 		}
@@ -399,7 +474,8 @@ namespace UsdToUnreal
 
 			for (const pxr::UsdPrim& Child : Prim.GetChildren())
 			{
-				bSuccess &= ConvertGroomHierarchy(Child, TimeCode, TransformToPropagate, HairDescription, AnimInfo);
+				const bool bCheckChildrenGroomAttributes = false;
+				bSuccess &= ConvertGroomHierarchy(Child, TimeCode, TransformToPropagate, HairDescription, AnimInfo, bCheckChildrenGroomAttributes);
 				if (!bSuccess)
 				{
 					break;
@@ -407,6 +483,15 @@ namespace UsdToUnreal
 			}
 		}
 
+		if (bCheckGroomAttributes)
+		{
+			// Groom-scoped attributes are checked only once, on the top-level prim
+			// Attributes from Alembic are stored as attributes in the userProperties namespace
+			// #ueent_todo: Will need to adjust once the USD groom schemas are defined
+			std::vector<pxr::UsdProperty> UserProperties = Prim.GetAuthoredPropertiesInNamespace("userProperties");
+			UE::UsdGroomConversion::Private::ConvertGroomAttributes(UserProperties, TimeCode, HairDescription);
+		}
+
 		return bSuccess;
 	}
 }	 // namespace UsdToUnreal
diff --git a/Engine/Plugins/Runtime/USDCore/Source/USDUtilities/Public/USDGroomConversion.h b/Engine/Plugins/Runtime/USDCore/Source/USDUtilities/Public/USDGroomConversion.h
--- a/Engine/Plugins/Runtime/USDCore/Source/USDUtilities/Public/USDGroomConversion.h
+++ b/Engine/Plugins/Runtime/USDCore/Source/USDUtilities/Public/USDGroomConversion.h
@@ -28,6 +28,8 @@ namespace UsdToUnreal
 	 * @param ParentTransform - The prim transform to propagate to the children
 	 * @param HairDescription - The hair description where the groom prim data is outputted
 	 * @param AnimInfo -  Optional GroomAnimationInfo to fill when parsing the hierarchy
+	 * @param bCheckGroomAttributes -  Optional Whether to check the groom-scope attributes, 
+	 *                                 should only be checked (true) on the top-level prim
 	 * @return true if the conversion was successful; false otherwise
 	 */
 	USDUTILITIES_API bool ConvertGroomHierarchy(
@@ -35,7 +37,8 @@ namespace UsdToUnreal
 		const pxr::UsdTimeCode& TimeCode,
 		const FTransform& ParentTransform,
 		FHairDescription& HairDescription,
-		FGroomAnimationInfo* AnimInfo = nullptr
+		FGroomAnimationInfo* AnimInfo = nullptr,
+		bool bCheckGroomAttributes = true
 	);
 }
 
