USD: Fixed return value for SetActiveLODIndex.

#rnx

#Synced-CL 51098007

diff --git a/Engine/Plugins/Runtime/USDCore/Source/UnrealUSDWrapper/Private/UnrealUSDWrapper.cpp b/Engine/Plugins/Runtime/USDCore/Source/UnrealUSDWrapper/Private/UnrealUSDWrapper.cpp
--- a/Engine/Plugins/Runtime/USDCore/Source/UnrealUSDWrapper/Private/UnrealUSDWrapper.cpp
+++ b/Engine/Plugins/Runtime/USDCore/Source/UnrealUSDWrapper/Private/UnrealUSDWrapper.cpp
@@ -406,10 +406,9 @@ bool IUsdPrim::SetActiveLODIndex(const pxr::UsdPrim& Prim, int LODIndex)
 		{
 			vector<string> VariantNames = LODVariantSet.GetVariantNames();
 
-			bool bResult = false;
 			if (LODIndex < VariantNames.size())
 			{
-				bResult = LODVariantSet.SetVariantSelection(VariantNames[LODIndex]);
+				return LODVariantSet.SetVariantSelection(VariantNames[LODIndex]);
 			}
 		}
 	}
