Interchange:
- Add method to set Enable Tessellation for material instance nodes, used by displacement

USD:
- Refactor parent materials using Substrate
- Add support Specular Workflow
- Add Support for Clearcoat, Clearcoatroughness, OpacityThreshold, OpacityMode and Displacement on the UsdPreviewSurface
- Enable the new parent material for project using Substrate, continue to fallback on legacy otherwise, or keep legacy DisplayColor parent material for vertex colors

MF_UsdPreviewSurface:
- Add support for opacityMode (transparent/presence)
- Fix F0 when specularColor is white making materials opaque and metallic when opacity is 0 (case of opacityMode is set to transparent)
- Replace Substrate horizontal blend between metal and specular by using Substrate select

#virtualized
#Synced-CL 51055552

diff --git a/Engine/Plugins/Interchange/Extensions/OpenUSD/Source/Import/Private/InterchangeUSDPipeline.cpp b/Engine/Plugins/Interchange/Extensions/OpenUSD/Source/Import/Private/InterchangeUSDPipeline.cpp
--- a/Engine/Plugins/Interchange/Extensions/OpenUSD/Source/Import/Private/InterchangeUSDPipeline.cpp
+++ b/Engine/Plugins/Interchange/Extensions/OpenUSD/Source/Import/Private/InterchangeUSDPipeline.cpp
@@ -2,21 +2,18 @@
 
 #include "InterchangeUSDPipeline.h"
 
+#include "InterchangeImportModule.h"
 #include "InterchangeGeometryCacheFactoryNode.h"
 #include "InterchangeHeterogeneousVolumeActorFactoryNode.h"
 #include "InterchangeMaterialDefinitions.h"
 #include "InterchangeMaterialFactoryNode.h"
 #include "InterchangeMaterialInstanceNode.h"
 #include "InterchangeMeshActorFactoryNode.h"
-#include "InterchangeMeshBundleNode.h"
 #include "InterchangeMeshFactoryNode.h"
 #include "InterchangeMeshNode.h"
-#include "InterchangeSceneComponentNodes.h"
 #include "InterchangeSceneNode.h"
 #include "InterchangeShaderGraphNode.h"
-#include "InterchangeSkeletalMeshLodDataNode.h"
 #include "InterchangeSparseVolumeTextureFactoryNode.h"
-#include "InterchangeStaticMeshLodDataNode.h"
 #include "InterchangeVolumeNode.h"
 #include "Nodes/InterchangeBaseNodeContainer.h"
 
@@ -40,6 +37,7 @@
 #include "USDConversionUtils.h"
 #include "USDErrorUtils.h"
 #include "UsdWrappers/SdfPath.h"
+#include "USDMaterialUtils.h"
 
 #if WITH_EDITORONLY_DATA
 #include "AssetUtils/Texture2DBuilder.h"
@@ -1406,12 +1404,59 @@ void UInterchangeUsdPipeline::ExecutePipeline(
 
 	Super::ExecutePipeline(NodeContainer, InSourceDatas, ContentBasePath);
 
+	if (IInterchangeImportModule::Get().IsSubstrateEnabled())
+	{
+		NodeContainer->IterateNodesOfType<UInterchangeMaterialInstanceNode>(
+			[NodeContainer](const FString& NodeUID, UInterchangeMaterialInstanceNode* MaterialInstanceNode)
+			{
+				FString ParentPath;
+				MaterialInstanceNode->GetCustomParent(ParentPath);
+				// Let's keep Vertex Color parent material for the moment (maybe we could instead set UseVertexColor in the parent material for Substrate?)
+				if (ParentPath.Contains("DisplayColor", ESearchCase::CaseSensitive))
+				{
+					return;
+				}
+
+				const FString FactoryNodeUID = UInterchangeFactoryBaseNode::BuildFactoryNodeUid(NodeUID);
+				if (UInterchangeMaterialInstanceFactoryNode* MaterialInstanceFactoryNode = Cast<UInterchangeMaterialInstanceFactoryNode>(NodeContainer->GetFactoryNode(FactoryNodeUID)))
+				{
+					MaterialInstanceFactoryNode->SetCustomParent(TEXT("/InterchangeAssets/Substrate/UsdPreviewSurface.UsdPreviewSurface"));
+
+					if (int32 Value; MaterialInstanceNode->GetInt32Attribute(UE::Interchange::USD::UsdReferenceMaterialProperties, Value))
+					{
+						EUsdReferenceMaterialProperties Properties = static_cast<EUsdReferenceMaterialProperties>(Value);
+						if (EnumHasAnyFlags(Properties, EUsdReferenceMaterialProperties::VT))
+						{
+							MaterialInstanceFactoryNode->SetCustomParent(TEXT("/InterchangeAssets/Substrate/UsdPreviewSurfaceVT.UsdPreviewSurfaceVT"));
+						}
+
+						if (EnumHasAnyFlags(Properties, EUsdReferenceMaterialProperties::Translucent))
+						{
+							MaterialInstanceFactoryNode->SetCustomBlendMode(EBlendMode::BLEND_TranslucentGreyTransmittance);
+						}
+
+						if (EnumHasAnyFlags(Properties, EUsdReferenceMaterialProperties::TwoSided))
+						{
+							MaterialInstanceFactoryNode->SetCustomTwoSided(true);
+						}
+
+						FString DisplacementTexture;
+						float Displacement;
+						if (MaterialInstanceNode->GetFloatAttribute(TEXT("Displacement"), Displacement) ||
+							MaterialInstanceNode->GetTextureParameterValue(TEXT("DisplacementTexture"), DisplacementTexture))
+						{
+							MaterialInstanceFactoryNode->SetCustomEnableTessellation(true);
+						}
+					}
+				}
+			});
+	}
+
 #if USE_USD_SDK
 	if (!NodeContainer)
 	{
 		return;
 	}
-
 	MoveCustomAttributesToMaterialFactoryNodes(*NodeContainer);
 
 	ProcessVolumeNodes(*NodeContainer);
diff --git a/Engine/Plugins/Interchange/Extensions/OpenUSD/Source/Import/Private/SchemaHandlers/UniversalMaterialSchemaHandler.cpp b/Engine/Plugins/Interchange/Extensions/OpenUSD/Source/Import/Private/SchemaHandlers/UniversalMaterialSchemaHandler.cpp
--- a/Engine/Plugins/Interchange/Extensions/OpenUSD/Source/Import/Private/SchemaHandlers/UniversalMaterialSchemaHandler.cpp
+++ b/Engine/Plugins/Interchange/Extensions/OpenUSD/Source/Import/Private/SchemaHandlers/UniversalMaterialSchemaHandler.cpp
@@ -14,11 +14,9 @@
 #include "USDErrorUtils.h"
 #include "USDShadeConversion.h"
 #include "UsdWrappers/UsdPrim.h"
-#include "UsdWrappers/UsdStage.h"
 
 #include "InterchangeMaterialInstanceNode.h"
 #include "InterchangeMaterialReferenceNode.h"
-#include "InterchangeSceneNode.h"
 #include "InterchangeShaderGraphNode.h"
 #include "InterchangeTranslatorHelper.h"
 #include "Nodes/InterchangeBaseNodeContainer.h"
@@ -216,6 +214,12 @@ namespace UE::UniversalMaterialSchemaHandler::Private
 			MaterialNode.AddScalarParameterValue(*BaseParameterName, static_cast<float>(Value));
 		}
 
+		void operator()(const int32 Value) const
+		{
+			// integers are treated as scalars
+			MaterialNode.AddScalarParameterValue(*BaseParameterName, static_cast<float>(Value));
+		}
+
 	public:
 		const UE::FUsdPrim& Prim;
 		UInterchangeMaterialInstanceNode& MaterialNode;
@@ -323,6 +327,7 @@ namespace UE::UniversalMaterialSchemaHandler::Private
 			Properties |= EUsdReferenceMaterialProperties::VT;
 		}
 
+		MaterialNode->AddInt32Attribute(UsdReferenceMaterialProperties, uint8(Properties));
 		FSoftObjectPath ParentMaterial = UsdUnreal::MaterialUtils::GetReferencePreviewSurfaceMaterial(Properties);
 		if (ParentMaterial.IsValid())
 		{
diff --git a/Engine/Plugins/Interchange/Runtime/Source/Common/Public/Usd/InterchangeUsdDefinitions.h b/Engine/Plugins/Interchange/Runtime/Source/Common/Public/Usd/InterchangeUsdDefinitions.h
--- a/Engine/Plugins/Interchange/Runtime/Source/Common/Public/Usd/InterchangeUsdDefinitions.h
+++ b/Engine/Plugins/Interchange/Runtime/Source/Common/Public/Usd/InterchangeUsdDefinitions.h
@@ -91,6 +91,9 @@ namespace UE::Interchange::USD
 
 	inline const FString PrimPathAttributeKeyString = TEXT("_USDPrimPathAttributeKey_");
 
+	// Attribute used to store the different properties of a material (Translucent, TwoSided, VT)
+	inline const FString UsdReferenceMaterialProperties = TEXT("UsdReferenceMaterialProperties");
+
 	// Custom attribute keys used to describe SVT info extracted from the USD custom schemas.
 	// We add these to the Volume nodes that the USD translator emits.
 	namespace SparseVolumeTexture
diff --git a/Engine/Plugins/Interchange/Runtime/Source/FactoryNodes/Private/InterchangeMaterialFactoryNode.cpp b/Engine/Plugins/Interchange/Runtime/Source/FactoryNodes/Private/InterchangeMaterialFactoryNode.cpp
--- a/Engine/Plugins/Interchange/Runtime/Source/FactoryNodes/Private/InterchangeMaterialFactoryNode.cpp
+++ b/Engine/Plugins/Interchange/Runtime/Source/FactoryNodes/Private/InterchangeMaterialFactoryNode.cpp
@@ -516,6 +516,16 @@ bool UInterchangeMaterialInstanceFactoryNode::SetCustomTwoSided(bool AttributeVa
 	IMPLEMENT_NODE_ATTRIBUTE_SETTER_NODELEGATE(TwoSided, bool);
 }
 
+bool UInterchangeMaterialInstanceFactoryNode::GetCustomEnableTessellation(bool& AttributeValue) const
+{
+	IMPLEMENT_NODE_ATTRIBUTE_GETTER(EnableTessellation, bool);
+}
+
+bool UInterchangeMaterialInstanceFactoryNode::SetCustomEnableTessellation(bool AttributeValue)
+{
+	IMPLEMENT_NODE_ATTRIBUTE_SETTER_NODELEGATE(EnableTessellation, bool);
+}
+
 FString UInterchangeMaterialReferenceFactoryNode::GetTypeName() const
 {
 	const FString TypeName = TEXT("MaterialReferenceFactoryNode");
diff --git a/Engine/Plugins/Interchange/Runtime/Source/FactoryNodes/Public/InterchangeMaterialFactoryNode.h b/Engine/Plugins/Interchange/Runtime/Source/FactoryNodes/Public/InterchangeMaterialFactoryNode.h
--- a/Engine/Plugins/Interchange/Runtime/Source/FactoryNodes/Public/InterchangeMaterialFactoryNode.h
+++ b/Engine/Plugins/Interchange/Runtime/Source/FactoryNodes/Public/InterchangeMaterialFactoryNode.h
@@ -347,11 +347,19 @@ public:
 	UFUNCTION(BlueprintCallable, Category = "Interchange | Node | MaterialInstanceFactory")
 	INTERCHANGEFACTORYNODES_API bool SetCustomTwoSided(bool AttributeValue);
 
+	/** Used to enable displacement*/
+	UFUNCTION(BlueprintCallable, Category = "Interchange | Node | MaterialInstanceFactory")
+	INTERCHANGEFACTORYNODES_API bool GetCustomEnableTessellation(bool& AttributeValue) const;
+
+	UFUNCTION(BlueprintCallable, Category = "Interchange | Node | MaterialInstanceFactory")
+	INTERCHANGEFACTORYNODES_API bool SetCustomEnableTessellation(bool AttributeValue);
+
 private:
 	IMPLEMENT_NODE_ATTRIBUTE_KEY(InstanceClassName)
 	IMPLEMENT_NODE_ATTRIBUTE_KEY(Parent)
 	IMPLEMENT_NODE_ATTRIBUTE_KEY(BlendMode)
 	IMPLEMENT_NODE_ATTRIBUTE_KEY(TwoSided)
+	IMPLEMENT_NODE_ATTRIBUTE_KEY(EnableTessellation)
 };
 
 /**
diff --git a/Engine/Plugins/Interchange/Runtime/Source/Import/Private/Material/InterchangeMaterialFactory.cpp b/Engine/Plugins/Interchange/Runtime/Source/Import/Private/Material/InterchangeMaterialFactory.cpp
--- a/Engine/Plugins/Interchange/Runtime/Source/Import/Private/Material/InterchangeMaterialFactory.cpp
+++ b/Engine/Plugins/Interchange/Runtime/Source/Import/Private/Material/InterchangeMaterialFactory.cpp
@@ -830,6 +830,20 @@ namespace UE::Interchange::MaterialFactory::Internal
 			MaterialInstance.BasePropertyOverrides.bOverride_TwoSided = false;
 		}
 	}
+
+	void OverrideEnableTessellation(UMaterialInstance& MaterialInstance, const UInterchangeMaterialInstanceFactoryNode& FactoryNode)
+	{
+		if (bool bEnableTessellation; FactoryNode.GetCustomEnableTessellation(bEnableTessellation))
+		{
+			MaterialInstance.BasePropertyOverrides.bOverride_bEnableTessellation = true;
+			MaterialInstance.BasePropertyOverrides.bEnableTessellation = bEnableTessellation;
+			MaterialInstance.bEnableTessellation = bEnableTessellation;
+		}
+		else
+		{
+			MaterialInstance.BasePropertyOverrides.bOverride_bEnableTessellation = false;
+		}
+	}
 }
 
 namespace UE::Interchange::Materials::HashUtils
@@ -2082,6 +2096,7 @@ void UInterchangeMaterialFactory::SetupMaterialInstance(UMaterialInstance& Mater
 
 	OverrideBlendMode(MaterialInstance, FactoryNode);
 	OverrideTwoSided(MaterialInstance, FactoryNode);
+	OverrideEnableTessellation(MaterialInstance, FactoryNode);
 }
 
 void UInterchangeMaterialFactory::SetupReimportedMaterialInstance(UMaterialInstance& MaterialInstance, const UInterchangeBaseNodeContainer& NodeContainer, const UInterchangeMaterialInstanceFactoryNode& FactoryNode, const UInterchangeMaterialInstanceFactoryNode& PreviousFactoryNode)
diff --git a/Engine/Plugins/Interchange/Runtime/Source/Nodes/Private/InterchangeMaterialInstanceNode.cpp b/Engine/Plugins/Interchange/Runtime/Source/Nodes/Private/InterchangeMaterialInstanceNode.cpp
--- a/Engine/Plugins/Interchange/Runtime/Source/Nodes/Private/InterchangeMaterialInstanceNode.cpp
+++ b/Engine/Plugins/Interchange/Runtime/Source/Nodes/Private/InterchangeMaterialInstanceNode.cpp
@@ -107,4 +107,14 @@ bool UInterchangeMaterialInstanceNode::GetCustomTwoSided(bool& AttributeValue) c
 bool UInterchangeMaterialInstanceNode::SetCustomTwoSided(bool AttributeValue)
 {
 	IMPLEMENT_NODE_ATTRIBUTE_SETTER_NODELEGATE(TwoSided, bool);
+}
+
+bool UInterchangeMaterialInstanceNode::GetCustomEnableTessellation(bool& AttributeValue) const
+{
+	IMPLEMENT_NODE_ATTRIBUTE_GETTER(EnableTessellation, bool);
+}
+
+bool UInterchangeMaterialInstanceNode::SetCustomEnableTessellation(bool AttributeValue)
+{
+	IMPLEMENT_NODE_ATTRIBUTE_SETTER_NODELEGATE(EnableTessellation, bool);
 }
\ No newline at end of file
diff --git a/Engine/Plugins/Interchange/Runtime/Source/Nodes/Public/InterchangeMaterialInstanceNode.h b/Engine/Plugins/Interchange/Runtime/Source/Nodes/Public/InterchangeMaterialInstanceNode.h
--- a/Engine/Plugins/Interchange/Runtime/Source/Nodes/Public/InterchangeMaterialInstanceNode.h
+++ b/Engine/Plugins/Interchange/Runtime/Source/Nodes/Public/InterchangeMaterialInstanceNode.h
@@ -77,11 +77,18 @@ public:
 	UFUNCTION(BlueprintCallable, Category = "Interchange | Node | Material")
 	UE_API bool SetCustomTwoSided(bool AttributeValue);
 
+	UFUNCTION(BlueprintCallable, Category = "Interchange | Node | Material")
+	UE_API bool GetCustomEnableTessellation(bool& AttributeValue) const;
+
+	UFUNCTION(BlueprintCallable, Category = "Interchange | Node | Material")
+	UE_API bool SetCustomEnableTessellation(bool AttributeValue);
+
 private:
 
 	IMPLEMENT_NODE_ATTRIBUTE_KEY(Parent);
 	IMPLEMENT_NODE_ATTRIBUTE_KEY(BlendMode);
 	IMPLEMENT_NODE_ATTRIBUTE_KEY(TwoSided);
+	IMPLEMENT_NODE_ATTRIBUTE_KEY(EnableTessellation);
 };
 
 #undef UE_API
diff --git a/Engine/Plugins/Interchange/Runtime/Source/Pipelines/Private/InterchangeGenericMaterialPipeline.cpp b/Engine/Plugins/Interchange/Runtime/Source/Pipelines/Private/InterchangeGenericMaterialPipeline.cpp
--- a/Engine/Plugins/Interchange/Runtime/Source/Pipelines/Private/InterchangeGenericMaterialPipeline.cpp
+++ b/Engine/Plugins/Interchange/Runtime/Source/Pipelines/Private/InterchangeGenericMaterialPipeline.cpp
@@ -1037,6 +1037,11 @@ void UInterchangeGenericMaterialPipeline::ExecutePipeline(UInterchangeBaseNodeCo
 			MaterialFactoryNode->SetCustomTwoSided(bTwoSided);
 		}
 
+		if (bool bEnableTessellation; MaterialNode->GetCustomEnableTessellation(bEnableTessellation))
+		{
+			MaterialFactoryNode->SetCustomEnableTessellation(bEnableTessellation);
+		}
+
 		UE::Interchange::InterchangeGenericMaterialPipeline::Private::AddSpecularProfileToFactoryNode(MaterialNode, MaterialFactoryNode, BaseNodeContainer);
 
 		TArray<FString> Inputs;
diff --git a/Engine/Plugins/Runtime/USDCore/Source/USDUtilities/Private/USDShadeConversion.cpp b/Engine/Plugins/Runtime/USDCore/Source/USDUtilities/Private/USDShadeConversion.cpp
--- a/Engine/Plugins/Runtime/USDCore/Source/USDUtilities/Private/USDShadeConversion.cpp
+++ b/Engine/Plugins/Runtime/USDCore/Source/USDUtilities/Private/USDShadeConversion.cpp
@@ -997,6 +997,43 @@ namespace UE::UsdShadeConversion::Private
 		return true;
 	}
 
+	bool GetIntParameterValue(
+		pxr::UsdShadeConnectableAPI& Connectable,
+		const pxr::TfToken& InputName,
+		int32 DefaultValue,
+		UsdToUnreal::FParameterValue& OutValue
+	)
+	{
+		FScopedUsdAllocs Allocs;
+
+		pxr::UsdShadeInput Input = Connectable.GetInput(InputName);
+		if (!Input)
+		{
+			return false;
+		}
+
+		// If we have another shader/node connected
+		pxr::UsdShadeConnectableAPI Source;
+		pxr::TfToken SourceName;
+		pxr::UsdShadeAttributeType AttributeType;
+		if (pxr::UsdShadeConnectableAPI::GetConnectedSource(Input.GetAttr(), &Source, &SourceName, &AttributeType))
+		{
+			// Recurse because the attribute may just be pointing at some other attribute that has the data
+			// (e.g. when shader input is just "hoisted" and connected to the parent material input)
+			return GetIntParameterValue(Source, SourceName, DefaultValue, OutValue);
+		}
+		// No other node connected, so we must have some value
+		else
+		{
+			int32 InputValue = DefaultValue;
+			Input.Get<int32>(&InputValue);
+
+			OutValue.Set<int32>(InputValue);
+		}
+
+		return true;
+	}
+
 #if WITH_EDITOR
 	/**
 	 * Intended to be used with Visit(), this will create a UMaterialExpression in InMaterial, set it with the
@@ -1116,6 +1153,12 @@ namespace UE::UsdShadeConversion::Private
 			UsdUtils::SetBoolParameterValue(Material, ParameterName, BoolValue);
 		}
 
+		// we handle int as scalar parameters
+		void operator()(const int32 Int32Value) const
+		{
+			UsdUtils::SetScalarParameterValue(Material, ParameterName, Int32Value);
+		}
+
 	protected:
 		UMaterialInstance& Material;
 		const TCHAR* ParameterName;
@@ -2234,12 +2277,48 @@ bool UsdToUnreal::ConvertMaterial(
 		OutMaterialData.Parameters.Add(TEXT("Roughness"), ParameterValue);
 	}
 
+	// SpecularColor
+	if (GetVec3ParameterValue(Connectable, UnrealIdentifiers::SpecularColor, FLinearColor::White, ParameterValue))
+	{
+		OutMaterialData.Parameters.Add(TEXT("SpecularColor"), ParameterValue);
+	}
+
+	// Use Specular Workflow
+	if (GetIntParameterValue(Connectable, UnrealIdentifiers::UseSpecularWorkflow, 0, ParameterValue))
+	{
+		OutMaterialData.Parameters.Add(TEXT("UseSpecularWorkflow"), ParameterValue);
+	}
+
+	// ClearCoat
+	if (GetFloatParameterValue(Connectable, UnrealIdentifiers::ClearCoat, 0, ParameterValue))
+	{
+		OutMaterialData.Parameters.Add(TEXT("ClearCoat"), ParameterValue);
+	}
+
+	// ClearCoat Roughness
+	if (GetFloatParameterValue(Connectable, UnrealIdentifiers::ClearCoatRoughness, 0.1, ParameterValue))
+	{
+		OutMaterialData.Parameters.Add(TEXT("ClearCoatRoughness"), ParameterValue);
+	}
+
 	// Opacity
 	if (GetFloatParameterValue(Connectable, UnrealIdentifiers::Opacity, 1.f, ParameterValue))
 	{
 		OutMaterialData.Parameters.Add(TEXT("Opacity"), ParameterValue);
 	}
 
+	// Opacity Mode
+	if (GetIntParameterValue(Connectable, UnrealIdentifiers::OpacityMode, 0, ParameterValue))
+	{
+		OutMaterialData.Parameters.Add(TEXT("OpacityMode"), ParameterValue);
+	}
+
+	// Opacity Threshold
+	if (GetFloatParameterValue(Connectable, UnrealIdentifiers::OpacityThreshold, 0, ParameterValue))
+	{
+		OutMaterialData.Parameters.Add(TEXT("OpacityThreshold"), ParameterValue);
+	}
+
 	// Normal
 	if (GetVec3ParameterValue(Connectable, UnrealIdentifiers::Normal, FLinearColor::Blue, ParameterValue, TEXTUREGROUP_WorldNormalMap))
 	{
@@ -2258,6 +2337,12 @@ bool UsdToUnreal::ConvertMaterial(
 		OutMaterialData.Parameters.Add(TEXT("AmbientOcclusion"), ParameterValue);
 	}
 
+	// Displacement
+	if (GetFloatParameterValue(Connectable, UnrealIdentifiers::Displacement, 0.0f, ParameterValue))
+	{
+		OutMaterialData.Parameters.Add(TEXT("Displacement"), ParameterValue);
+	}
+
 	TMap<FString, FString> ParameterToPrimvar;
 	TMap<FString, int32> PrimvarToUVIndex;
 
diff --git a/Engine/Plugins/Runtime/USDCore/Source/USDUtilities/Public/USDShadeConversion.h b/Engine/Plugins/Runtime/USDCore/Source/USDUtilities/Public/USDShadeConversion.h
--- a/Engine/Plugins/Runtime/USDCore/Source/USDUtilities/Public/USDShadeConversion.h
+++ b/Engine/Plugins/Runtime/USDCore/Source/USDUtilities/Public/USDShadeConversion.h
@@ -78,7 +78,7 @@ namespace UsdToUnreal
 		FVector FallbackValue;
 	};
 
-	using FParameterValue = TVariant<float, FVector, FTextureParameterValue, FPrimvarReaderParameterValue, bool>;
+	using FParameterValue = TVariant<float, FVector, FTextureParameterValue, FPrimvarReaderParameterValue, bool, int32>;
 
 	struct FUsdPreviewSurfaceMaterialData
 	{
diff --git a/Engine/Plugins/Runtime/USDCore/Source/UnrealUSDWrapper/Private/UnrealUSDWrapper.cpp b/Engine/Plugins/Runtime/USDCore/Source/UnrealUSDWrapper/Private/UnrealUSDWrapper.cpp
--- a/Engine/Plugins/Runtime/USDCore/Source/UnrealUSDWrapper/Private/UnrealUSDWrapper.cpp
+++ b/Engine/Plugins/Runtime/USDCore/Source/UnrealUSDWrapper/Private/UnrealUSDWrapper.cpp
@@ -130,11 +130,18 @@ namespace UnrealIdentifiers
 	const TfToken Metallic = TfToken("metallic");
 	const TfToken Roughness = TfToken("roughness");
 	const TfToken Opacity = TfToken("opacity");
+	const TfToken OpacityMode = TfToken("opacityMode");
+	const TfToken OpacityThreshold = TfToken("opacityThreshold");
 	const TfToken Normal = TfToken("normal");
 	const TfToken Specular = TfToken("specular");
+	const TfToken SpecularColor = TfToken("specularColor");
+	const TfToken UseSpecularWorkflow = TfToken("useSpecularWorkflow");
+	const TfToken ClearCoat = TfToken("clearcoat");
+	const TfToken ClearCoatRoughness = TfToken("clearcoatRoughness");
 	const TfToken Anisotropy = TfToken("anisotropy");
 	const TfToken Tangent = TfToken("tangent");
 	const TfToken SubsurfaceColor = TfToken("subsurfaceColor");
+	const TfToken Displacement = TfToken("displacement");
 	const TfToken Occlusion = TfToken("occlusion");
 	const TfToken Refraction = TfToken("ior");
 
diff --git a/Engine/Plugins/Runtime/USDCore/Source/UnrealUSDWrapper/Public/UnrealUSDWrapper.h b/Engine/Plugins/Runtime/USDCore/Source/UnrealUSDWrapper/Public/UnrealUSDWrapper.h
--- a/Engine/Plugins/Runtime/USDCore/Source/UnrealUSDWrapper/Public/UnrealUSDWrapper.h
+++ b/Engine/Plugins/Runtime/USDCore/Source/UnrealUSDWrapper/Public/UnrealUSDWrapper.h
@@ -526,11 +526,18 @@ namespace UnrealIdentifiers
 	extern UNREALUSDWRAPPER_API const pxr::TfToken Metallic;
 	extern UNREALUSDWRAPPER_API const pxr::TfToken Roughness;
 	extern UNREALUSDWRAPPER_API const pxr::TfToken Opacity;
+	extern UNREALUSDWRAPPER_API const pxr::TfToken OpacityMode;
+	extern UNREALUSDWRAPPER_API const pxr::TfToken OpacityThreshold;
 	extern UNREALUSDWRAPPER_API const pxr::TfToken Normal;
 	extern UNREALUSDWRAPPER_API const pxr::TfToken Specular;
+	extern UNREALUSDWRAPPER_API const pxr::TfToken SpecularColor;
+	extern UNREALUSDWRAPPER_API const pxr::TfToken UseSpecularWorkflow;
+	extern UNREALUSDWRAPPER_API const pxr::TfToken ClearCoat;
+	extern UNREALUSDWRAPPER_API const pxr::TfToken ClearCoatRoughness;
 	extern UNREALUSDWRAPPER_API const pxr::TfToken Anisotropy;
 	extern UNREALUSDWRAPPER_API const pxr::TfToken Tangent;
 	extern UNREALUSDWRAPPER_API const pxr::TfToken SubsurfaceColor;
+	extern UNREALUSDWRAPPER_API const pxr::TfToken Displacement;
 	extern UNREALUSDWRAPPER_API const pxr::TfToken Occlusion;
 	extern UNREALUSDWRAPPER_API const pxr::TfToken Refraction;
 
