USD:
- Read opacityMode as a token instead as an integer
- Avoid overwriting the parent material on material instance nodes if it's a shadergraph node.

#Synced-CL 51290462

diff --git a/Engine/Plugins/Interchange/Extensions/OpenUSD/Source/Import/Private/InterchangeUSDPipeline.cpp b/Engine/Plugins/Interchange/Extensions/OpenUSD/Source/Import/Private/InterchangeUSDPipeline.cpp
--- a/Engine/Plugins/Interchange/Extensions/OpenUSD/Source/Import/Private/InterchangeUSDPipeline.cpp
+++ b/Engine/Plugins/Interchange/Extensions/OpenUSD/Source/Import/Private/InterchangeUSDPipeline.cpp
@@ -1404,11 +1404,26 @@ void UInterchangeUsdPipeline::ExecutePipeline(
 
 	Super::ExecutePipeline(NodeContainer, InSourceDatas, ContentBasePath);
 
+	if (!NodeContainer)
+	{
+		UE_LOG(LogInterchangeUSDPipeline, Warning, TEXT("UInterchangeUSDPipeline: Cannot execute pre-import pipeline because NodeContrainer is null"));
+		return;
+	}
+
 	if (IInterchangeImportModule::Get().IsSubstrateEnabled())
 	{
 		NodeContainer->IterateNodesOfType<UInterchangeMaterialInstanceNode>(
 			[NodeContainer](const FString& NodeUID, UInterchangeMaterialInstanceNode* MaterialInstanceNode)
 			{
+				// We can have a parent material that comes from a ShaderGraphNode (for example a TwoSided instance created based on the ShaderGraph),
+				// in that case it's not a UsdPreviewSurface
+				bool bShaderGaphParent = false;
+				MaterialInstanceNode->GetBooleanAttribute(UE::Interchange::USD::UsdMaterialInstanceFromShaderGraph, bShaderGaphParent);
+				if (bShaderGaphParent)
+				{
+					return;
+				}
+
 				FString ParentPath;
 				MaterialInstanceNode->GetCustomParent(ParentPath);
 				// Let's keep Vertex Color parent material for the moment (maybe we could instead set UseVertexColor in the parent material for Substrate?)
diff --git a/Engine/Plugins/Interchange/Extensions/OpenUSD/Source/Import/Private/SchemaHandlers/SchemaHandlerUtils.cpp b/Engine/Plugins/Interchange/Extensions/OpenUSD/Source/Import/Private/SchemaHandlers/SchemaHandlerUtils.cpp
--- a/Engine/Plugins/Interchange/Extensions/OpenUSD/Source/Import/Private/SchemaHandlers/SchemaHandlerUtils.cpp
+++ b/Engine/Plugins/Interchange/Extensions/OpenUSD/Source/Import/Private/SchemaHandlers/SchemaHandlerUtils.cpp
@@ -1114,6 +1114,7 @@ namespace UE::Interchange::USD
 		UInterchangeMaterialInstanceNode* TwoSidedNode = NewObject<UInterchangeMaterialInstanceNode>(&NodeContainer);
 		NodeContainer.SetupNode(TwoSidedNode, TwoSidedUid, TwoSidedNodeName, EInterchangeNodeContainerType::TranslatedAsset);
 		TwoSidedNode->SetCustomParent(MaterialNode->GetUniqueID());
+		TwoSidedNode->AddBooleanAttribute(UsdMaterialInstanceFromShaderGraph, true);
 		TwoSidedNode->SetCustomTwoSided(true);
 		return TwoSidedNode;
 	}
diff --git a/Engine/Plugins/Interchange/Runtime/Source/Common/Public/Usd/InterchangeUsdDefinitions.h b/Engine/Plugins/Interchange/Runtime/Source/Common/Public/Usd/InterchangeUsdDefinitions.h
--- a/Engine/Plugins/Interchange/Runtime/Source/Common/Public/Usd/InterchangeUsdDefinitions.h
+++ b/Engine/Plugins/Interchange/Runtime/Source/Common/Public/Usd/InterchangeUsdDefinitions.h
@@ -94,6 +94,9 @@ namespace UE::Interchange::USD
 	// Attribute used to store the different properties of a material (Translucent, TwoSided, VT)
 	inline const FString UsdReferenceMaterialProperties = TEXT("UsdReferenceMaterialProperties");
 
+	// Attribute to avoid overwriting parent materials by the UsdPreviewSurface
+	inline const FString UsdMaterialInstanceFromShaderGraph = TEXT("UsdMaterialInstanceFromShaderGraph");
+
 	// Custom attribute keys used to describe SVT info extracted from the USD custom schemas.
 	// We add these to the Volume nodes that the USD translator emits.
 	namespace SparseVolumeTexture
diff --git a/Engine/Plugins/Runtime/USDCore/Source/USDUtilities/Private/USDShadeConversion.cpp b/Engine/Plugins/Runtime/USDCore/Source/USDUtilities/Private/USDShadeConversion.cpp
--- a/Engine/Plugins/Runtime/USDCore/Source/USDUtilities/Private/USDShadeConversion.cpp
+++ b/Engine/Plugins/Runtime/USDCore/Source/USDUtilities/Private/USDShadeConversion.cpp
@@ -1034,6 +1034,35 @@ namespace UE::UsdShadeConversion::Private
 		return true;
 	}
 
+	// No recursion for token value
+	bool GetTokenAsFloatParameterValue(
+		pxr::UsdShadeConnectableAPI& Connectable,
+		const pxr::TfToken& InputName,
+		const char* DefaultValue,
+		const TMap<FString, float>& Values,
+		UsdToUnreal::FParameterValue& OutValue
+	)
+	{
+		FScopedUsdAllocs Allocs;
+
+		pxr::UsdShadeInput Input = Connectable.GetInput(InputName);
+		if (!Input)
+		{
+			return false;
+		}
+
+		pxr::TfToken TokenInputValue{ DefaultValue };
+		Input.Get<pxr::TfToken>(&TokenInputValue);
+		float InputValue = 0;
+		if (const float* FoundValue = Values.Find(UsdToUnreal::ConvertToken(TokenInputValue)))
+		{
+			InputValue = *FoundValue;
+		}
+		OutValue.Set<float>(InputValue);
+
+		return true;
+	}
+
 #if WITH_EDITOR
 	/**
 	 * Intended to be used with Visit(), this will create a UMaterialExpression in InMaterial, set it with the
@@ -2308,7 +2337,8 @@ bool UsdToUnreal::ConvertMaterial(
 	}
 
 	// Opacity Mode
-	if (GetIntParameterValue(Connectable, UnrealIdentifiers::OpacityMode, 0, ParameterValue))
+	static TMap<FString, float> OpacityModeMap{ {"transparent", 0.f}, {"presence", 1.f} };
+	if (GetTokenAsFloatParameterValue(Connectable, UnrealIdentifiers::OpacityMode, "transparent", OpacityModeMap, ParameterValue))
 	{
 		OutMaterialData.Parameters.Add(TEXT("OpacityMode"), ParameterValue);
 	}
