USD Interchange: Handle UseWorldSpaceNormals material parameter.

#Synced-CL 51330716

diff --git a/Engine/Plugins/Runtime/USDCore/Source/USDUtilities/Private/USDShadeConversion.cpp b/Engine/Plugins/Runtime/USDCore/Source/USDUtilities/Private/USDShadeConversion.cpp
--- a/Engine/Plugins/Runtime/USDCore/Source/USDUtilities/Private/USDShadeConversion.cpp
+++ b/Engine/Plugins/Runtime/USDCore/Source/USDUtilities/Private/USDShadeConversion.cpp
@@ -2373,6 +2373,20 @@ bool UsdToUnreal::ConvertMaterial(
 		OutMaterialData.Parameters.Add(TEXT("Displacement"), ParameterValue);
 	}
 
+	// UseWorldSpaceNormal
+	// This is not part of the UsdPreviewSurface schema, but it's something we export for landscape materials.
+	// Unlike the other attributes this is written directly to the Material prim, and not on the surface shader.
+	if (pxr::UsdAttribute Attr = UsdShadeMaterialPrim.GetAttribute(UnrealIdentifiers::WorldSpaceNormals))
+	{
+		bool bValue = false;
+		if (Attr.Get<bool>(&bValue) && bValue)
+		{
+			// Use an int parameter here as actual bool parameters are used for static switches it seems
+			UsdToUnreal::FParameterValue IntParameterValue{TInPlaceType<int32>(), static_cast<int32>(bValue)};
+			OutMaterialData.Parameters.Add(TEXT("UseWorldSpaceNormals"), IntParameterValue);
+		}
+	}
+
 	TMap<FString, FString> ParameterToPrimvar;
 	TMap<FString, int32> PrimvarToUVIndex;
 
@@ -2432,17 +2446,6 @@ bool UsdToUnreal::ConvertMaterial(
 		UsdShadeConversionImpl::SetParameterValue(Material, *Parameter.Key, Parameter.Value, bForUsdPreviewSurface, MaterialData.PrimvarToUVIndex);
 	}
 
-	// Handle world space normals
-	pxr::UsdPrim UsdPrim = UsdShadeMaterial.GetPrim();
-	if (pxr::UsdAttribute Attr = UsdPrim.GetAttribute(UnrealIdentifiers::WorldSpaceNormals))
-	{
-		bool bValue = false;
-		if (Attr.Get<bool>(&bValue) && bValue)
-		{
-			UsdUtils::SetScalarParameterValue(Material, TEXT("UseWorldSpaceNormals"), 1.0f);
-		}
-	}
-
 	// Record which primvars we used on each UV index. This is important as we'll match this with the analogous member
 	// on static/skeletal meshe import data, and create a new material instance with different UV index parameter
 	// values if we need to
