USD: Fixed inputs:normal from USD shader not importing if it's a single value.

#Synced-CL 51543066

diff --git a/Engine/Plugins/Runtime/USDCore/Source/USDUtilities/Private/USDShadeConversion.cpp b/Engine/Plugins/Runtime/USDCore/Source/USDUtilities/Private/USDShadeConversion.cpp
--- a/Engine/Plugins/Runtime/USDCore/Source/USDUtilities/Private/USDShadeConversion.cpp
+++ b/Engine/Plugins/Runtime/USDCore/Source/USDUtilities/Private/USDShadeConversion.cpp
@@ -925,7 +925,7 @@ namespace UE::UsdShadeConversion::Private
 				return GetVec3ParameterValue(Source, SourceName, DefaultValue, OutValue, Group);
 			}
 			// Has a value directly on the port
-			else if (InputName != UnrealIdentifiers::Normal)
+			else
 			{
 				FLinearColor DiffuseColor = DefaultValue;
 
