USD: Add support for exporting Refraction channel.

Always export refraction=1.0 as a fallback for translucent materials in case there's no other opinion, as otherwise the USD default differs from the UE default.

#Synced-CL 52058314

diff --git a/Engine/Plugins/Runtime/USDCore/Source/USDUtilities/Private/USDShadeConversion.cpp b/Engine/Plugins/Runtime/USDCore/Source/USDUtilities/Private/USDShadeConversion.cpp
--- a/Engine/Plugins/Runtime/USDCore/Source/USDUtilities/Private/USDShadeConversion.cpp
+++ b/Engine/Plugins/Runtime/USDCore/Source/USDUtilities/Private/USDShadeConversion.cpp
@@ -1715,7 +1715,17 @@ namespace UE::UsdShadeConversion::Private
 		{
 			PropertiesToProcess.Add(MP_BaseColor);
 		}
+		// Analogously, the default refraction value in USD is 1.5, which is very distorted, while it's 1.0 in UE.
+		// So if we have no other opinion here, let's author a 1.0 refraction value so that it looks like it would in UE.
+		// We only care about that if we're also baking opacity though (so that we're in some kind of translucent mode). 
+		// If we're fully opaque then we don't really care to author this
+		if (!PropertiesToProcess.Contains(MP_Refraction) && (PropertiesToProcess.Contains(MP_Opacity) || PropertiesToProcess.Contains(MP_OpacityMask)))
+		{
+			PropertiesToProcess.Add(MP_Refraction);
+		}
+
 		const float Zero = 0.0f;
+		const float One = 1.0f;
 
 		// Fill in outputs
 		for (const EMaterialProperty& Property : PropertiesToProcess)
@@ -1797,6 +1807,12 @@ namespace UE::UsdShadeConversion::Private
 				bPropertyValueIsConstant = true;
 			}
 
+			if (Property == MP_Refraction && NumSamples == 0 && !UserConstantValue)
+			{
+				UserConstantValue = &One;
+				bPropertyValueIsConstant = true;
+			}
+
 			if ((NumSamples == 0 || !SampleSize) && !UserConstantValue)
 			{
 				USD_LOG_INFO(TEXT("Skipping material property %d as we have no valid data to use."), Property);
@@ -1928,6 +1944,12 @@ namespace UE::UsdShadeConversion::Private
 					ConstantValue = UserConstantValue ? *UserConstantValue : ConstantLinearValue[0];
 					FallbackValue = pxr::GfVec4f{1.0, 1.0, 1.0, 1.0f};
 					break;
+				case MP_Refraction:
+					InputToken = UnrealIdentifiers::Refraction;
+					InputType = pxr::SdfValueTypeNames->Float;
+					ConstantValue = UserConstantValue ? *UserConstantValue : ConstantLinearValue[0];
+					FallbackValue = pxr::GfVec4f{1.0, 1.0, 1.0, 1.0f};
+					break;
 				case MP_Anisotropy:
 					InputToken = UnrealIdentifiers::Anisotropy;
 					InputType = pxr::SdfValueTypeNames->Float;
