USD: Fix LogManager message identifier producing FNames with negative integers, which now trigger asserts after cl 52066218

#Synced-CL 52058314

diff --git a/Engine/Plugins/Runtime/USDCore/Source/USDUtilities/Private/USDErrorUtils.cpp b/Engine/Plugins/Runtime/USDCore/Source/USDUtilities/Private/USDErrorUtils.cpp
--- a/Engine/Plugins/Runtime/USDCore/Source/USDUtilities/Private/USDErrorUtils.cpp
+++ b/Engine/Plugins/Runtime/USDCore/Source/USDUtilities/Private/USDErrorUtils.cpp
@@ -49,7 +49,12 @@ namespace UE::USDErrorUtils::Private
 
 	FName GenerateMessageIdentifier(uint32 Number)
 	{
-		return FName(TEXT("USD_LOG_ID"), Number);
+		// This FName constructor takes an int32. Unfortunately, there are now some new validation 
+		// on FNames that asserts if we try using negative integers, so here we mask off the top bit
+		// to ensure the int32 becomes positive.
+		//
+		// See also FName::IsValidInternalNumber and NAME_NO_NUMBER_INTERNAL
+		return FName(TEXT("USD_LOG_ID"), Number & 0x7FFFFFFFu);
 	}
 
 	void SendMessageToOutputLog(const FTokenizedMessage& Message)
