USD: Add fallback for light attributes being authored without the "inputs:" prefix  (for backwards compatibility with older scenes)

#Synced-CL 52058314

diff --git a/Engine/Plugins/Interchange/Extensions/OpenUSD/Source/Import/Private/SchemaHandlers/LightSchemaHandler.cpp b/Engine/Plugins/Interchange/Extensions/OpenUSD/Source/Import/Private/SchemaHandlers/LightSchemaHandler.cpp
--- a/Engine/Plugins/Interchange/Extensions/OpenUSD/Source/Import/Private/SchemaHandlers/LightSchemaHandler.cpp
+++ b/Engine/Plugins/Interchange/Extensions/OpenUSD/Source/Import/Private/SchemaHandlers/LightSchemaHandler.cpp
@@ -24,6 +24,7 @@
 #include "Nodes/InterchangeBaseNodeContainer.h"
 
 #include "Engine/Scene.h"
+#include "InterchangeCommonAnimationPayload.h"
 #include "Misc/Paths.h"
 
 #include "USDIncludesStart.h"
@@ -109,16 +110,16 @@ namespace UE::LuxLightSchemaHandler::Private
 
 		// Ref. UsdToUnreal::ConvertLight
 		static const FString IntensityToken = UsdToUnreal::ConvertToken(pxr::UsdLuxTokens->inputsIntensity);
-		float IntensityUsd = 1.0f;
-		bool bIntensityAuthored = UsdUtils::GetAuthoredAttributeValue<float>(Prim, IntensityToken, IntensityUsd);
+		bool bIntensityAuthored = false;
+		float IntensityUsd = UsdUtils::GetLightAttrValueWithInputsFallback<float>(Prim.GetAttribute(*IntensityToken), UsdUtils::GetDefaultTimeCode(), &bIntensityAuthored);
 
+		bool bExposureAuthored = false;
 		static const FString ExposureToken = UsdToUnreal::ConvertToken(pxr::UsdLuxTokens->inputsExposure);
-		float Exposure = 0.0f;
-		UsdUtils::GetAuthoredAttributeValue<float>(Prim, ExposureToken, Exposure);
+		float Exposure = UsdUtils::GetLightAttrValueWithInputsFallback<float>(Prim.GetAttribute(*ExposureToken), UsdUtils::GetDefaultTimeCode(), &bExposureAuthored);
 
 		static const FString NormalizeToken = UsdToUnreal::ConvertToken(pxr::UsdLuxTokens->inputsNormalize);
-		bool bNormalizedIntensityUsd = false;
-		bool bNormalizeAuthored = UsdUtils::GetAuthoredAttributeValue<bool>(Prim, NormalizeToken, bNormalizedIntensityUsd);
+		bool bNormalizeAuthored = false;
+		bool bNormalizedIntensityUsd = UsdUtils::GetLightAttrValueWithInputsFallback<bool>(Prim.GetAttribute(*NormalizeToken), UsdUtils::GetDefaultTimeCode(), &bNormalizeAuthored);
 		bNormalizedIntensityUsd = bNormalizeAuthored && bNormalizedIntensityUsd;
 
 		static const FString RadiusToken = UsdToUnreal::ConvertToken(pxr::UsdLuxTokens->inputsRadius);
@@ -128,6 +129,16 @@ namespace UE::LuxLightSchemaHandler::Private
 		FLinearColor Color = FLinearColor::White;
 		static const FString ColorToken = UsdToUnreal::ConvertToken(pxr::UsdLuxTokens->inputsColor);
 		UE::FUsdAttribute ColorAttribute = Prim.GetAttribute(*ColorToken);
+		if (!ColorAttribute || !ColorAttribute.HasAuthoredValue())
+		{
+			// Fallback to the unprefixed attribute for pre-21.05 USD scenes
+			static const FString UnprefixedColorToken = TEXT("color");
+			UE::FUsdAttribute FallbackAttr = Prim.GetAttribute(*UnprefixedColorToken);
+			if (FallbackAttr)
+			{
+				ColorAttribute = FallbackAttr;
+			}
+		}
 		if (ColorAttribute && ColorAttribute.HasAuthoredValue())
 		{
 			UE::FVtValue Value;
@@ -166,13 +177,17 @@ namespace UE::LuxLightSchemaHandler::Private
 		}
 
 		static const FString TemperatureToken = UsdToUnreal::ConvertToken(pxr::UsdLuxTokens->inputsColorTemperature);
-		if (float Temperature; UsdUtils::GetAuthoredAttributeValue<float>(Prim, TemperatureToken, Temperature))
+		bool bTemperatureAuthored = false;
+		float Temperature = UsdUtils::GetLightAttrValueWithInputsFallback<float>(Prim.GetAttribute(*TemperatureToken), UsdUtils::GetDefaultTimeCode(), &bTemperatureAuthored); 
+		if (bTemperatureAuthored)
 		{
 			LightNode->SetCustomTemperature(Temperature);
 		}
 
 		static const FString UseTemperatureToken = UsdToUnreal::ConvertToken(pxr::UsdLuxTokens->inputsEnableColorTemperature);
-		if (bool UseTemperature; UsdUtils::GetAuthoredAttributeValue<bool>(Prim, UseTemperatureToken, UseTemperature))
+		bool bUseTemperatureAuthored = false;
+		bool UseTemperature = UsdUtils::GetLightAttrValueWithInputsFallback<bool>(Prim.GetAttribute(*UseTemperatureToken), UsdUtils::GetDefaultTimeCode(), &bUseTemperatureAuthored); 
+		if (bUseTemperatureAuthored)
 		{
 			LightNode->SetCustomUseTemperature(UseTemperature);
 		}
@@ -193,7 +208,9 @@ namespace UE::LuxLightSchemaHandler::Private
 			bAllowNormalizingIntensity = false;
 
 			static const FString SourceAngleToken = UsdToUnreal::ConvertToken(pxr::UsdLuxTokens->inputsAngle);
-			if (float UsdAngle; UsdUtils::GetAuthoredAttributeValue<float>(Prim, SourceAngleToken, UsdAngle))
+			bool bAngleAuthored = false;
+			float UsdAngle = UsdUtils::GetLightAttrValueWithInputsFallback<float>(Prim.GetAttribute(*SourceAngleToken), UsdUtils::GetDefaultTimeCode(), &bAngleAuthored); 
+			if (bAngleAuthored)
 			{
 				CastLightNode->SetCustomLightSourceAngle(UsdAngle);
 			}
@@ -205,7 +222,9 @@ namespace UE::LuxLightSchemaHandler::Private
 			float RadiusUE = 1.0f;
 			if (UInterchangePointLightNode* CastLightNode = Cast<UInterchangePointLightNode>(LightNode))
 			{
-				if (float RadiusUsd = 1.0f; UsdUtils::GetAuthoredAttributeValue<float>(Prim, RadiusToken, RadiusUsd))
+				bool bRadiusAuthored = false;
+				float RadiusUsd = UsdUtils::GetLightAttrValueWithInputsFallback<float>(Prim.GetAttribute(*RadiusToken), UsdUtils::GetDefaultTimeCode(), &bRadiusAuthored); 
+				if (bRadiusAuthored)
 				{
 					RadiusUE = UsdToUnreal::ConvertDistance(StageInfo, RadiusUsd);
 					CastLightNode->SetCustomSourceRadius(RadiusUE);
@@ -224,12 +243,12 @@ namespace UE::LuxLightSchemaHandler::Private
 				static const FString ConeSoftnessToken = UsdToUnreal::ConvertToken(pxr::UsdLuxTokens->inputsShapingConeSoftness);
 
 				// Default from UsdLuxShapingAPI
-				float UsdConeAngle = 90.0f;
-				bool bConeAngleAuthored = UsdUtils::GetAuthoredAttributeValue<float>(Prim, ConeAngleToken, UsdConeAngle);
+				bool bConeAngleAuthored = false;
+				float UsdConeAngle = UsdUtils::GetLightAttrValueWithInputsFallback<float>(Prim.GetAttribute(*ConeAngleToken), UsdUtils::GetDefaultTimeCode(), &bConeAngleAuthored);
 
 				// Default from UsdLuxShapingAPI
-				float UsdConeSoftness = 0.0f;
-				bool bConeSoftnessAuthored = UsdUtils::GetAuthoredAttributeValue<float>(Prim, ConeSoftnessToken, UsdConeSoftness);
+				bool bConeSoftnessAuthored = false;
+				float UsdConeSoftness = UsdUtils::GetLightAttrValueWithInputsFallback<float>(Prim.GetAttribute(*ConeSoftnessToken), UsdUtils::GetDefaultTimeCode(), &bConeSoftnessAuthored);
 
 				float InnerConeAngle = 0.0f;
 				float OuterConeAngle = 44.0f; // Default from the USpotLightComponent constructor
@@ -266,14 +285,18 @@ namespace UE::LuxLightSchemaHandler::Private
 			if (Prim.IsA(TEXT("RectLight")))
 			{
 				float WidthUE = 1.0f;
-				if (float WidthUsd = 1.0f; UsdUtils::GetAuthoredAttributeValue<float>(Prim, WidthToken, WidthUsd))
+				bool bWidthAuthored = false;
+				float WidthUsd = UsdUtils::GetLightAttrValueWithInputsFallback<float>(Prim.GetAttribute(*WidthToken), UsdUtils::GetDefaultTimeCode(), &bWidthAuthored);
+				if (bWidthAuthored)
 				{
 					WidthUE = UsdToUnreal::ConvertDistance(StageInfo, WidthUsd);
 					CastLightNode->SetCustomSourceWidth(WidthUE);
 				}
 
 				float HeightUE = 1.0f;
-				if (float HeightUsd = 1.0f; UsdUtils::GetAuthoredAttributeValue<float>(Prim, HeightToken, HeightUsd))
+				bool bHeightAuthored = false;
+				float HeightUsd = UsdUtils::GetLightAttrValueWithInputsFallback<float>(Prim.GetAttribute(*HeightToken), UsdUtils::GetDefaultTimeCode(), &bHeightAuthored); 
+				if (bHeightAuthored)
 				{
 					HeightUE = UsdToUnreal::ConvertDistance(StageInfo, HeightUsd);
 					CastLightNode->SetCustomSourceHeight(HeightUE);
@@ -285,7 +308,9 @@ namespace UE::LuxLightSchemaHandler::Private
 			else // DiskLight
 			{
 				float RadiusUE = 0.5f;
-				if (float RadiusUsd = 1.0f; UsdUtils::GetAuthoredAttributeValue<float>(Prim, RadiusToken, RadiusUsd))
+				bool bDiskRadiusAuthored = false;
+				float RadiusUsd = UsdUtils::GetLightAttrValueWithInputsFallback<float>(Prim.GetAttribute(*RadiusToken), UsdUtils::GetDefaultTimeCode(), &bDiskRadiusAuthored); 
+				if (bDiskRadiusAuthored)
 				{
 					RadiusUE = UsdToUnreal::ConvertDistance(StageInfo, RadiusUsd);
 					float SquareSideUE = UsdUtils::GetRectLightSideFromDiskLight(RadiusUE);
@@ -312,6 +337,15 @@ namespace UE::LuxLightSchemaHandler::Private
 			// Show some warning in case the file specifies a different DomeLight format, as we only support latlong
 			static const FString FormatToken = UsdToUnreal::ConvertToken(pxr::UsdLuxTokens->inputsTextureFormat);
 			UE::FUsdAttribute FormatAttribute = Prim.GetAttribute(*FormatToken);
+			if (!FormatAttribute || !FormatAttribute.HasAuthoredValue())
+			{
+				static const FString UnprefixedFormatToken = TEXT("textureFormat");
+				UE::FUsdAttribute FallbackAttr = Prim.GetAttribute(*UnprefixedFormatToken);
+				if (FallbackAttr)
+				{
+					FormatAttribute = FallbackAttr;
+				}
+			}
 			FString Format;
 			if (FormatAttribute.HasAuthoredValue() &&
 				FormatAttribute.Get<FString>(Format) &&
@@ -331,6 +365,15 @@ namespace UE::LuxLightSchemaHandler::Private
 
 			static const FString FileToken = UsdToUnreal::ConvertToken(pxr::UsdLuxTokens->inputsTextureFile);
 			UE::FUsdAttribute FileAttribute = Prim.GetAttribute(*FileToken);
+			if (!FileAttribute || !FileAttribute.HasAuthoredValue())
+			{
+				static const FString UnprefixedFileToken = TEXT("textureFile");
+				UE::FUsdAttribute FallbackAttr = Prim.GetAttribute(*UnprefixedFileToken);
+				if (FallbackAttr)
+				{
+					FileAttribute = FallbackAttr;
+				}
+			}
 			const FString ResolvedDomeTexturePath = UsdUtils::GetResolvedAssetPath(FileAttribute);
 
 			FString TextureCubeNodeUid = AddTextureCubeNode(ResolvedDomeTexturePath, AccumulatedInfo, NodeContainer);
@@ -360,7 +403,7 @@ namespace UE::LuxLightSchemaHandler::Private
 			}
 		}
 
-		if (bIntensityAuthored)
+		if (bIntensityAuthored || bExposureAuthored)
 		{
 			// This will end up being Nits for local lights (point, rect, disk, spot), lux for distant lights,
 			// and a simple multiplier without units for dome lights
diff --git a/Engine/Plugins/Runtime/USDCore/Source/USDUtilities/Private/USDLightConversion.cpp b/Engine/Plugins/Runtime/USDCore/Source/USDUtilities/Private/USDLightConversion.cpp
--- a/Engine/Plugins/Runtime/USDCore/Source/USDUtilities/Private/USDLightConversion.cpp
+++ b/Engine/Plugins/Runtime/USDCore/Source/USDUtilities/Private/USDLightConversion.cpp
@@ -59,6 +59,57 @@ namespace LightConversionImpl
 	}
 }	 // namespace LightConversionImpl
 
+template<typename T>
+T UsdUtils::GetLightAttrValueWithInputsFallback(const pxr::UsdAttribute& Attribute, double UsdTimeCode, bool* bOutIsAuthored)
+{
+	if (bOutIsAuthored)
+	{
+		*bOutIsAuthored = false;
+	}
+
+	if (!Attribute)
+	{
+		return {};
+	}
+
+	pxr::UsdTimeCode TimeCode{UsdTimeCode};
+	pxr::UsdAttribute AttributeToUse = Attribute;
+	bool bIsAuthored = Attribute.GetResolveInfo(TimeCode).GetSource() != pxr::UsdResolveInfoSourceFallback;
+
+	if (!bIsAuthored)
+	{
+		std::string AttrName = Attribute.GetName();
+		const static std::string Prefix = "inputs:";
+		if (AttrName.starts_with(Prefix))
+		{
+			AttrName = AttrName.substr(Prefix.length());
+			pxr::UsdAttribute FallbackAttr = Attribute.GetPrim().GetAttribute(pxr::TfToken{AttrName});
+			if (FallbackAttr)
+			{
+				AttributeToUse = FallbackAttr;
+				bIsAuthored = FallbackAttr.GetResolveInfo(TimeCode).GetSource() != pxr::UsdResolveInfoSourceFallback;
+			}
+		}
+	}
+
+	if (bOutIsAuthored)
+	{
+		*bOutIsAuthored = bIsAuthored;
+	}
+
+	T Result;
+	if (AttributeToUse && AttributeToUse.Get(&Result, TimeCode))
+	{
+		return Result;
+	}
+
+	return {};
+}
+
+template USDUTILITIES_API float UsdUtils::GetLightAttrValueWithInputsFallback<float>(const pxr::UsdAttribute&, double, bool*);
+template USDUTILITIES_API bool UsdUtils::GetLightAttrValueWithInputsFallback<bool>(const pxr::UsdAttribute&, double, bool*);
+template USDUTILITIES_API pxr::GfVec3f UsdUtils::GetLightAttrValueWithInputsFallback<pxr::GfVec3f>(const pxr::UsdAttribute&, double, bool*);
+
 bool UsdToUnreal::ConvertLight(const pxr::UsdPrim& Prim, ULightComponentBase& LightComponentBase, double UsdTimeCode)
 {
 	FScopedUsdAllocs Allocs;
@@ -69,9 +120,9 @@ bool UsdToUnreal::ConvertLight(const pxr::UsdPrim& Prim, ULightComponentBase& Li
 		return false;
 	}
 
-	const float UsdIntensity = UsdUtils::GetUsdValue<float>(LightAPI.GetIntensityAttr(), UsdTimeCode);
-	const float UsdExposure = UsdUtils::GetUsdValue<float>(LightAPI.GetExposureAttr(), UsdTimeCode);
-	const pxr::GfVec3f UsdColor = UsdUtils::GetUsdValue<pxr::GfVec3f>(LightAPI.GetColorAttr(), UsdTimeCode);
+	const float UsdIntensity = UsdUtils::GetLightAttrValueWithInputsFallback<float>(LightAPI.GetIntensityAttr(), UsdTimeCode);
+	const float UsdExposure = UsdUtils::GetLightAttrValueWithInputsFallback<float>(LightAPI.GetExposureAttr(), UsdTimeCode);
+	const pxr::GfVec3f UsdColor = UsdUtils::GetLightAttrValueWithInputsFallback<pxr::GfVec3f>(LightAPI.GetColorAttr(), UsdTimeCode);
 
 	const bool bSRGB = true;
 	LightComponentBase.LightColor = UsdToUnreal::ConvertColor(UsdColor).ToFColor(bSRGB);
@@ -79,20 +130,13 @@ bool UsdToUnreal::ConvertLight(const pxr::UsdPrim& Prim, ULightComponentBase& Li
 
 	if (ULightComponent* LightComponent = Cast<ULightComponent>(&LightComponentBase))
 	{
-		LightComponent->bUseTemperature = UsdUtils::GetUsdValue<bool>(LightAPI.GetEnableColorTemperatureAttr(), UsdTimeCode);
-		LightComponent->Temperature = UsdUtils::GetUsdValue<float>(LightAPI.GetColorTemperatureAttr(), UsdTimeCode);
+		LightComponent->bUseTemperature = UsdUtils::GetLightAttrValueWithInputsFallback<bool>(LightAPI.GetEnableColorTemperatureAttr(), UsdTimeCode);
+		LightComponent->Temperature = UsdUtils::GetLightAttrValueWithInputsFallback<float>(LightAPI.GetColorTemperatureAttr(), UsdTimeCode);
 	}
 
 	if (const pxr::UsdLuxShadowAPI ShadowAPI{Prim})
 	{
-		if (pxr::UsdAttribute Attr = ShadowAPI.GetShadowEnableAttr())
-		{
-			bool bEnable = true;
-			if (Attr.Get(&bEnable, UsdTimeCode))
-			{
-				LightComponentBase.SetCastShadows(bEnable);
-			}
-		}
+		LightComponentBase.SetCastShadows(UsdUtils::GetLightAttrValueWithInputsFallback<bool>(ShadowAPI.GetShadowEnableAttr(), UsdTimeCode));
 	}
 
 	return true;
@@ -108,7 +152,7 @@ bool UsdToUnreal::ConvertDistantLight(const pxr::UsdPrim& Prim, UDirectionalLigh
 		return false;
 	}
 
-	LightComponent.LightSourceAngle = UsdUtils::GetUsdValue<float>(DistantLight.GetAngleAttr(), UsdTimeCode);
+	LightComponent.LightSourceAngle = UsdUtils::GetLightAttrValueWithInputsFallback<float>(DistantLight.GetAngleAttr(), UsdTimeCode);
 
 	return true;
 }
@@ -125,11 +169,11 @@ bool UsdToUnreal::ConvertRectLight(const pxr::UsdPrim& Prim, URectLightComponent
 
 	const FUsdStageInfo StageInfo(Prim.GetStage());
 
-	const float UsdIntensity = UsdUtils::GetUsdValue<float>(RectLight.GetIntensityAttr(), UsdTimeCode);
-	const float UsdExposure = UsdUtils::GetUsdValue<float>(RectLight.GetExposureAttr(), UsdTimeCode);
-	const float UsdWidth = UsdUtils::GetUsdValue<float>(RectLight.GetWidthAttr(), UsdTimeCode);
-	const float UsdHeight = UsdUtils::GetUsdValue<float>(RectLight.GetHeightAttr(), UsdTimeCode);
-	const double UsdNormalized = UsdUtils::GetUsdValue<bool>(RectLight.GetNormalizeAttr(), pxr::UsdTimeCode::Default());
+	const float UsdIntensity = UsdUtils::GetLightAttrValueWithInputsFallback<float>(RectLight.GetIntensityAttr(), UsdTimeCode);
+	const float UsdExposure = UsdUtils::GetLightAttrValueWithInputsFallback<float>(RectLight.GetExposureAttr(), UsdTimeCode);
+	const float UsdWidth = UsdUtils::GetLightAttrValueWithInputsFallback<float>(RectLight.GetWidthAttr(), UsdTimeCode);
+	const float UsdHeight = UsdUtils::GetLightAttrValueWithInputsFallback<float>(RectLight.GetHeightAttr(), UsdTimeCode);
+	const double UsdNormalized = UsdUtils::GetLightAttrValueWithInputsFallback<bool>(RectLight.GetNormalizeAttr(), pxr::UsdTimeCode::Default().GetValue());
 
 	const double NitsIntensity = UsdToUnreal::ConvertLightIntensity(UsdIntensity, UsdExposure);
 
@@ -163,10 +207,10 @@ bool UsdToUnreal::ConvertDiskLight(const pxr::UsdPrim& Prim, URectLightComponent
 
 	const FUsdStageInfo StageInfo(Prim.GetStage());
 
-	const float UsdIntensity = UsdUtils::GetUsdValue<float>(DiskLight.GetIntensityAttr(), UsdTimeCode);
-	const float UsdExposure = UsdUtils::GetUsdValue<float>(DiskLight.GetExposureAttr(), UsdTimeCode);
-	const float UsdRadius = UsdUtils::GetUsdValue<float>(DiskLight.GetRadiusAttr(), UsdTimeCode);
-	const double UsdNormalized = UsdUtils::GetUsdValue<bool>(DiskLight.GetNormalizeAttr(), pxr::UsdTimeCode::Default());
+	const float UsdIntensity = UsdUtils::GetLightAttrValueWithInputsFallback<float>(DiskLight.GetIntensityAttr(), UsdTimeCode);
+	const float UsdExposure = UsdUtils::GetLightAttrValueWithInputsFallback<float>(DiskLight.GetExposureAttr(), UsdTimeCode);
+	const float UsdRadius = UsdUtils::GetLightAttrValueWithInputsFallback<float>(DiskLight.GetRadiusAttr(), UsdTimeCode);
+	const double UsdNormalized = UsdUtils::GetLightAttrValueWithInputsFallback<bool>(DiskLight.GetNormalizeAttr(), pxr::UsdTimeCode::Default().GetValue());
 
 	const float RadiusUE = UsdToUnreal::ConvertDistance(StageInfo, UsdRadius);
 	const float SquareSideUE = UsdUtils::GetRectLightSideFromDiskLight(RadiusUE);
@@ -203,10 +247,11 @@ bool UsdToUnreal::ConvertSphereLight(const pxr::UsdPrim& Prim, UPointLightCompon
 
 	const FUsdStageInfo StageInfo(Prim.GetStage());
 
-	const float UsdIntensity = UsdUtils::GetUsdValue<float>(SphereLight.GetIntensityAttr(), UsdTimeCode);
-	const float UsdExposure = UsdUtils::GetUsdValue<float>(SphereLight.GetExposureAttr(), UsdTimeCode);
-	const float UsdRadius = UsdUtils::GetUsdValue<float>(SphereLight.GetRadiusAttr(), UsdTimeCode);
-	const double UsdNormalized = UsdUtils::GetUsdValue<bool>(SphereLight.GetNormalizeAttr(), pxr::UsdTimeCode::Default());
+	const float UsdIntensity = UsdUtils::GetLightAttrValueWithInputsFallback<float>(SphereLight.GetIntensityAttr(), UsdTimeCode);
+	const float UsdExposure = UsdUtils::GetLightAttrValueWithInputsFallback<float>(SphereLight.GetExposureAttr(), UsdTimeCode);
+	const float UsdRadius = UsdUtils::GetLightAttrValueWithInputsFallback<float>(SphereLight.GetRadiusAttr(), UsdTimeCode);
+	const double UsdNormalized = UsdUtils::GetLightAttrValueWithInputsFallback<bool>(SphereLight.GetNormalizeAttr(), pxr::UsdTimeCode::Default().GetValue());
+	const double TreatAsPoint = UsdUtils::GetUsdValue<bool>(SphereLight.GetTreatAsPointAttr(), UsdTimeCode);
 
 	const float RadiusUE = UsdToUnreal::ConvertDistance(StageInfo, UsdRadius);
 	const double NitsIntensity = UsdToUnreal::ConvertLightIntensity(UsdIntensity, UsdExposure);
@@ -226,8 +271,8 @@ bool UsdToUnreal::ConvertSphereLight(const pxr::UsdPrim& Prim, UPointLightCompon
 			pxr::UsdLuxShapingAPI ShapingAPI{Prim};
 			if (ShapingAPI)
 			{
-				const float UsdConeAngle = UsdUtils::GetUsdValue<float>(ShapingAPI.GetShapingConeAngleAttr(), UsdTimeCode);
-				const float UsdConeSoftness = UsdUtils::GetUsdValue<float>(ShapingAPI.GetShapingConeSoftnessAttr(), UsdTimeCode);
+				const float UsdConeAngle = UsdUtils::GetLightAttrValueWithInputsFallback<float>(ShapingAPI.GetShapingConeAngleAttr(), UsdTimeCode);
+				const float UsdConeSoftness = UsdUtils::GetLightAttrValueWithInputsFallback<float>(ShapingAPI.GetShapingConeSoftnessAttr(), UsdTimeCode);
 				float InnerConeAngle = 0.0f;
 				float OuterConeAngle = 1.0f;
 				UsdToUnreal::ConvertSpotLightConeAngles(UsdConeAngle, UsdConeSoftness, OuterConeAngle, InnerConeAngle);
@@ -269,8 +314,8 @@ bool UsdToUnreal::ConvertLuxShapingAPI(const pxr::UsdPrim& Prim, USpotLightCompo
 		return false;
 	}
 
-	const float UsdConeAngle = UsdUtils::GetUsdValue<float>(ShapingAPI.GetShapingConeAngleAttr(), UsdTimeCode);
-	const float UsdConeSoftness = UsdUtils::GetUsdValue<float>(ShapingAPI.GetShapingConeSoftnessAttr(), UsdTimeCode);
+	const float UsdConeAngle = UsdUtils::GetLightAttrValueWithInputsFallback<float>(ShapingAPI.GetShapingConeAngleAttr(), UsdTimeCode);
+	const float UsdConeSoftness = UsdUtils::GetLightAttrValueWithInputsFallback<float>(ShapingAPI.GetShapingConeSoftnessAttr(), UsdTimeCode);
 
 	float InnerConeAngle = 0.0f;
 	float OuterConeAngle = 1.0f;
diff --git a/Engine/Plugins/Runtime/USDCore/Source/USDUtilities/Public/USDLightConversion.h b/Engine/Plugins/Runtime/USDCore/Source/USDUtilities/Public/USDLightConversion.h
--- a/Engine/Plugins/Runtime/USDCore/Source/USDUtilities/Public/USDLightConversion.h
+++ b/Engine/Plugins/Runtime/USDCore/Source/USDUtilities/Public/USDLightConversion.h
@@ -22,6 +22,7 @@ PXR_NAMESPACE_OPEN_SCOPE
 	class UsdLuxShapingAPI;
 	class UsdLuxSphereLight;
 
+	class UsdAttribute;
 	class UsdPrim;
 	class UsdStage;
 	template<typename T>
@@ -206,6 +207,21 @@ namespace UsdUtils
 
 	/** For a light with the provided SolidAngle and surface AreaInSqMeters, converts the Intensity value in Nits into the provided DestUnits */
 	USDUTILITIES_API float ConvertIntensityFromNits(float Intensity, float SolidAngle, float AreaInSqMeters, ELightUnits DestUnits);
+
+	/**
+	 * Reads the value of Attribute at the given time code. If the attribute has no authored opinion
+	 * and its name starts with "inputs:", strips the prefix and tries the unprefixed attribute.
+	 * This handles backwards compatibility with pre-21.05 USD scenes.
+	 *
+	 * If bOutIsAuthored is non-null, sets it to true if the value came from an authored opinion
+	 * (either on the original attribute or the fallback), false if it's only a schema default.
+	 */
+	template<typename T>
+	USDUTILITIES_API T GetLightAttrValueWithInputsFallback(
+		const pxr::UsdAttribute& Attribute,
+		double UsdTimeCode = UsdUtils::GetDefaultTimeCode(),
+		bool* bOutIsAuthored = nullptr
+	);
 }
 
 #endif	  // #if USE_USD_SDK
