USD: Bake all frames of exposure animation tracks.

#Synced-CL 52058314

diff --git a/Engine/Plugins/Importers/USDImporter/Source/USDStage/Private/USDLevelSequenceHelper.cpp b/Engine/Plugins/Importers/USDImporter/Source/USDStage/Private/USDLevelSequenceHelper.cpp
--- a/Engine/Plugins/Importers/USDImporter/Source/USDStage/Private/USDLevelSequenceHelper.cpp
+++ b/Engine/Plugins/Importers/USDImporter/Source/USDStage/Private/USDLevelSequenceHelper.cpp
@@ -2769,6 +2769,8 @@ void FUsdLevelSequenceHelperImpl::AddLightTracks(const UUsdPrimTwin& PrimTwin, c
 			continue;
 		}
 
+		UsdUtils::DensifyTimeSamplesIfNeeded(Prim, PropertyPath, UnionedTimeSamples);
+
 		UE::FSdfLayer AttributeLayer;
 		ULevelSequence* AttributeSequence = FindOrAddSequenceForAttribute(MainAttr, &AttributeLayer);
 		if (!AttributeSequence)
diff --git a/Engine/Plugins/Interchange/Extensions/OpenUSD/Source/Import/Private/SchemaHandlers/ImageableSchemaHandler.cpp b/Engine/Plugins/Interchange/Extensions/OpenUSD/Source/Import/Private/SchemaHandlers/ImageableSchemaHandler.cpp
--- a/Engine/Plugins/Interchange/Extensions/OpenUSD/Source/Import/Private/SchemaHandlers/ImageableSchemaHandler.cpp
+++ b/Engine/Plugins/Interchange/Extensions/OpenUSD/Source/Import/Private/SchemaHandlers/ImageableSchemaHandler.cpp
@@ -399,6 +399,8 @@ namespace UE::GeomImageableSchemaHandler::Private
 			return false;
 		}
 
+		UsdUtils::DensifyTimeSamplesIfNeeded(Prim, UEPropertyName, TimeSampleUnion);
+
 		const bool bIgnorePrimLocalTransform = false;
 		Reader = UsdToUnreal::CreatePropertyTrackReader(Prim, UEPropertyName, bIgnorePrimLocalTransform);
 
diff --git a/Engine/Plugins/Runtime/USDCore/Source/USDUtilities/Private/USDPrimConversion.cpp b/Engine/Plugins/Runtime/USDCore/Source/USDUtilities/Private/USDPrimConversion.cpp
--- a/Engine/Plugins/Runtime/USDCore/Source/USDUtilities/Private/USDPrimConversion.cpp
+++ b/Engine/Plugins/Runtime/USDCore/Source/USDUtilities/Private/USDPrimConversion.cpp
@@ -6876,6 +6876,89 @@ TArray<FName> UsdUtils::GetPropertiesForAttribute(const UE::FUsdPrim& Prim, cons
 	return {};
 }
 
+void UsdUtils::DensifyTimeSamplesIfNeeded(
+	const UE::FUsdPrim& Prim,
+	const FName& PropertyPath,
+	TArray<double>& InOutTimeSamples
+)
+{
+	using namespace UnrealIdentifiers;
+
+	// For now, only light intensity needs densification, because the final USD intensity
+	// is intensity * 2^exposure, which is non-linear in exposure. If we only sample at
+	// authored keyframe times, linearly interpolating the resulting UE keyframes won't
+	// match the non-linear USD interpolation of the separate attributes.
+	if (PropertyPath != IntensityPropertyName)
+	{
+		return;
+	}
+
+	if (InOutTimeSamples.Num() < 2)
+	{
+		return;
+	}
+
+	FScopedUsdAllocs Allocs;
+	pxr::UsdPrim UsdPrim{Prim};
+
+	pxr::UsdAttribute ExposureAttr = UsdPrim.GetAttribute(pxr::UsdLuxTokens->inputsExposure);
+	if (!ExposureAttr || !ExposureAttr.ValueMightBeTimeVarying())
+	{
+		return;
+	}
+
+	// Check if all exposure time samples are zero (effectively no exposure animation)
+	std::vector<double> ExposureSamples;
+	ExposureAttr.GetTimeSamples(&ExposureSamples);
+	if (ExposureSamples.empty())
+	{
+		return;
+	}
+
+	bool bAllZero = true;
+	for (double SampleTime : ExposureSamples)
+	{
+		float ExposureValue = 0.0f;
+		if (ExposureAttr.Get<float>(&ExposureValue, SampleTime) && !FMath::IsNearlyZero(ExposureValue))
+		{
+			bAllZero = false;
+			break;
+		}
+	}
+	if (bAllZero)
+	{
+		return;
+	}
+
+	// Densify: generate every integer time code between the first and last sample
+	USD_LOG_INFO(
+		TEXT("Baking light intensity animation frame-by-frame for prim '%s' because it has animated exposure, "
+			 "which combines non-linearly with intensity (intensity * 2^exposure)"),
+		*Prim.GetPrimPath().GetString()
+	);
+
+	const double First = FMath::FloorToDouble(InOutTimeSamples[0]);
+	const double Last = FMath::CeilToDouble(InOutTimeSamples.Last());
+
+	TSet<double> TimeCodeSet;
+	TimeCodeSet.Reserve(static_cast<int32>(Last - First) + 1 + InOutTimeSamples.Num());
+
+	// Preserve original samples (they may not be on integer boundaries)
+	for (double Sample : InOutTimeSamples)
+	{
+		TimeCodeSet.Add(Sample);
+	}
+
+	// Add every integer time code in the range
+	for (double TimeCode = First; TimeCode <= Last; TimeCode += 1.0)
+	{
+		TimeCodeSet.Add(TimeCode);
+	}
+
+	InOutTimeSamples = TimeCodeSet.Array();
+	InOutTimeSamples.Sort();
+}
+
 bool UnrealToUsd::ConvertDrawModeComponent(
 	const UUsdDrawModeComponent& DrawModeComponent,
 	pxr::UsdPrim& UsdPrim,
diff --git a/Engine/Plugins/Runtime/USDCore/Source/USDUtilities/Public/USDPrimConversion.h b/Engine/Plugins/Runtime/USDCore/Source/USDUtilities/Public/USDPrimConversion.h
--- a/Engine/Plugins/Runtime/USDCore/Source/USDUtilities/Public/USDPrimConversion.h
+++ b/Engine/Plugins/Runtime/USDCore/Source/USDUtilities/Public/USDPrimConversion.h
@@ -522,6 +522,18 @@ namespace UsdUtils
 	 *          [UnrealIdentifiers::SourceWidthPropertyName, UnrealIdentifiers::SourceHeightPropertyName].
 	 */
 	USDUTILITIES_API TArray<FName> GetPropertiesForAttribute(const UE::FUsdPrim& Prim, const FString& AttrName);
+
+	/**
+	 * Checks if the given property requires frame-by-frame baking due to non-linear attribute
+	 * interactions (e.g. light intensity = inputs:intensity * 2^inputs:exposure).
+	 * If so, densifies InOutTimeSamples to include every integer time code between the first and last
+	 * existing sample, ensuring the non-linear combination is properly captured as linear keyframes.
+	 */
+	USDUTILITIES_API void DensifyTimeSamplesIfNeeded(
+		const UE::FUsdPrim& Prim,
+		const FName& PropertyPath,
+		TArray<double>& InOutTimeSamples
+	);
 }
 
 #endif	  // #if USE_USD_SDK
