USD: Clamp source radius to a minimum value to ensure light units converted to/from normalized units don't lead to zero/infinities.

#Synced-CL 52058314

diff --git a/Engine/Plugins/Runtime/USDCore/Source/USDUtilities/Private/USDLightConversion.cpp b/Engine/Plugins/Runtime/USDCore/Source/USDUtilities/Private/USDLightConversion.cpp
--- a/Engine/Plugins/Runtime/USDCore/Source/USDUtilities/Private/USDLightConversion.cpp
+++ b/Engine/Plugins/Runtime/USDCore/Source/USDUtilities/Private/USDLightConversion.cpp
@@ -44,6 +44,17 @@
 
 #define LOCTEXT_NAMESPACE "USDLightConversion"
 
+static float GMinLightSourceSize = 1.0f;
+static FAutoConsoleVariableRef CVarMinLightSourceSize(
+	TEXT("USD.LightExport.MinLightSourceSize"),
+	GMinLightSourceSize,
+	TEXT("Minimum source size (in centimeters) for point, spot, and rect lights when exporting to USD. "
+		 "Lights with source dimensions below this value will have their dimensions clamped to this minimum "
+		 "before writing to USD, preventing extreme intensity values from near-zero surface area. "
+		 "For point/spot lights this clamps SourceRadius; for rect lights this clamps SourceWidth and SourceHeight. "
+		 "Set to 0 to disable clamping. Default: 1.0")
+);
+
 namespace LightConversionImpl
 {
 	// Copied from USpotLightComponent::GetCosHalfConeAngle, so we don't need a component to do the same math
@@ -489,8 +500,10 @@ bool UnrealToUsd::ConvertRectLightComponent(const URectLightComponent& LightComp
 
 	FUsdStageInfo StageInfo(Prim.GetStage());
 
-	const float WidthMeters = LightComponent.SourceWidth / 100.0f;
-	const float HeightMeters = LightComponent.SourceHeight / 100.0f;
+	const float ClampedWidth = FMath::Max(LightComponent.SourceWidth, GMinLightSourceSize);
+	const float ClampedHeight = FMath::Max(LightComponent.SourceHeight, GMinLightSourceSize);
+	const float WidthMeters = ClampedWidth / 100.0f;
+	const float HeightMeters = ClampedHeight / 100.0f;
 	const float AreaSqMeters = UsdUtils::GetRectLightArea(WidthMeters, HeightMeters);
 
 	float SolidAngle = 1.0f;
@@ -502,7 +515,7 @@ bool UnrealToUsd::ConvertRectLightComponent(const URectLightComponent& LightComp
 
 		if (pxr::UsdAttribute Attr = DiskLight.CreateRadiusAttr())
 		{
-			const float DiskRadiusUEUnits = UsdUtils::GetDiskLightRadiusFromRectLight(LightComponent.SourceWidth, LightComponent.SourceHeight);
+			const float DiskRadiusUEUnits = UsdUtils::GetDiskLightRadiusFromRectLight(ClampedWidth, ClampedHeight);
 			Attr.Set<float>(UnrealToUsd::ConvertDistance(StageInfo, DiskRadiusUEUnits), UsdTimeCode);
 			UsdUtils::NotifyIfOverriddenOpinion(Attr);
 		}
@@ -514,13 +527,13 @@ bool UnrealToUsd::ConvertRectLightComponent(const URectLightComponent& LightComp
 
 		if (pxr::UsdAttribute Attr = RectLight.CreateWidthAttr())
 		{
-			Attr.Set<float>(UnrealToUsd::ConvertDistance(StageInfo, LightComponent.SourceWidth), UsdTimeCode);
+			Attr.Set<float>(UnrealToUsd::ConvertDistance(StageInfo, ClampedWidth), UsdTimeCode);
 			UsdUtils::NotifyIfOverriddenOpinion(Attr);
 		}
 
 		if (pxr::UsdAttribute Attr = RectLight.CreateHeightAttr())
 		{
-			Attr.Set<float>(UnrealToUsd::ConvertDistance(StageInfo, LightComponent.SourceHeight), UsdTimeCode);
+			Attr.Set<float>(UnrealToUsd::ConvertDistance(StageInfo, ClampedHeight), UsdTimeCode);
 			UsdUtils::NotifyIfOverriddenOpinion(Attr);
 		}
 	}
@@ -567,7 +580,7 @@ bool UnrealToUsd::ConvertPointLightComponent(const UPointLightComponent& LightCo
 
 	FUsdStageInfo StageInfo(Prim.GetStage());
 
-	const float RadiusUE = LightComponent.SourceRadius;
+	const float RadiusUE = FMath::Max(LightComponent.SourceRadius, GMinLightSourceSize);
 	const float RadiusMeters = RadiusUE / 100.0f;
 
 	if (pxr::UsdAttribute Attr = Light.CreateRadiusAttr())
@@ -578,7 +591,7 @@ bool UnrealToUsd::ConvertPointLightComponent(const UPointLightComponent& LightCo
 
 	if (pxr::UsdAttribute Attr = Light.CreateTreatAsPointAttr())
 	{
-		Attr.Set<bool>(FMath::IsNearlyZero(RadiusUE), UsdTimeCode);
+		Attr.Set<bool>(FMath::IsNearlyZero(LightComponent.SourceRadius), UsdTimeCode);
 		UsdUtils::NotifyIfOverriddenOpinion(Attr);
 	}
 
diff --git a/Engine/Plugins/Runtime/USDCore/Source/USDUtilities/Private/USDPrimConversion.cpp b/Engine/Plugins/Runtime/USDCore/Source/USDUtilities/Private/USDPrimConversion.cpp
--- a/Engine/Plugins/Runtime/USDCore/Source/USDUtilities/Private/USDPrimConversion.cpp
+++ b/Engine/Plugins/Runtime/USDCore/Source/USDUtilities/Private/USDPrimConversion.cpp
@@ -5051,6 +5051,9 @@ bool UnrealToUsd::CreateComponentPropertyBaker(
 		TFunction<void(float, double)> TemperatureSetter = CreateCachedAttrSetter(TemperatureAttr, NoFloatConversion);
 		TFunction<void(bool, double)> ShadowSetter = CreateCachedAttrSetter(ShadowEnableAttr, NoBoolConversion);
 
+		static IConsoleVariable* MinLightSourceSizeCvar = IConsoleManager::Get().FindConsoleVariable(TEXT("USD.LightExport.MinLightSourceSize"));
+		const float MinLightSourceSize = MinLightSourceSizeCvar ? MinLightSourceSizeCvar->GetFloat() : 1.0f;
+
 		if (const URectLightComponent* RectLightComponent = Cast<URectLightComponent>(LightComponentBase))
 		{
 			static TSet<FString> RelevantProperties =
@@ -5080,7 +5083,8 @@ bool UnrealToUsd::CreateComponentPropertyBaker(
 								 TemperatureSetter,
 								 ShadowSetter,
 								 WidthSetter,
-								 HeightSetter](double UsdTimeCode) mutable
+								 HeightSetter,
+								 MinLightSourceSize](double UsdTimeCode) mutable
 				{
 					TRACE_CPUPROFILER_EVENT_SCOPE(RectLightBaker);
 
@@ -5090,25 +5094,20 @@ bool UnrealToUsd::CreateComponentPropertyBaker(
 						return;
 					}
 
-					float UEIntensity = RectLightComponent->Intensity;
-
-					float AreaInSqMeters = (RectLightComponent->SourceWidth / 100.0f) * (RectLightComponent->SourceHeight / 100.0f);
-					if (FMath::IsNearlyZero(AreaInSqMeters))
-					{
-						UEIntensity = 0.0f;
-					}
-
-					const float Steradians = PI;
+					const float SourceWidthUE = FMath::Max(RectLightComponent->SourceWidth, MinLightSourceSize);
+					const float SourceHeightUE = FMath::Max(RectLightComponent->SourceHeight, MinLightSourceSize);
+					const float SolidAngle = UsdUtils::GetRectLightSolidAngle();
+					const float AreaInSqMeters = UsdUtils::GetRectLightArea(SourceWidthUE / 100.0f, SourceHeightUE / 100.0f);
 					const float FinalIntensityNits = UsdUtils::ConvertIntensityToNits(
-						UEIntensity,
-						Steradians,
+						RectLightComponent->Intensity,
+						SolidAngle,
 						AreaInSqMeters,
 						RectLightComponent->IntensityUnits
 					);
 
 					IntensitySetter(FinalIntensityNits, UsdTimeCode);
-					WidthSetter(RectLightComponent->SourceWidth, UsdTimeCode);
-					HeightSetter(RectLightComponent->SourceHeight, UsdTimeCode);
+					WidthSetter(SourceWidthUE, UsdTimeCode);
+					HeightSetter(SourceHeightUE, UsdTimeCode);
 
 					ColorSetter(RectLightComponent->LightColor, UsdTimeCode);
 					EnableTemperatureSetter(RectLightComponent->bUseTemperature, UsdTimeCode);
@@ -5154,7 +5153,8 @@ bool UnrealToUsd::CreateComponentPropertyBaker(
 								 TemperatureSetter,
 								 ShadowSetter,
 								 AngleSetter,
-								 SoftnessSetter](double UsdTimeCode) mutable
+								 SoftnessSetter,
+								 MinLightSourceSize](double UsdTimeCode) mutable
 				{
 					TRACE_CPUPROFILER_EVENT_SCOPE(SpotLightBaker);
 
@@ -5164,8 +5164,9 @@ bool UnrealToUsd::CreateComponentPropertyBaker(
 						return;
 					}
 
-					const float SolidAngle = 2.f * PI * (1.0f - SpotLightComponent->GetCosHalfConeAngle());
-					const float AreaInSqMeters = FMath::Max(SolidAngle * FMath::Square(SpotLightComponent->SourceRadius / 100.f), KINDA_SMALL_NUMBER);
+					const float SourceRadiusUE = FMath::Max(SpotLightComponent->SourceRadius, MinLightSourceSize);
+					const float SolidAngle = UsdUtils::GetSpotLightSolidAngle(SpotLightComponent->OuterConeAngle, SpotLightComponent->InnerConeAngle);
+					const float AreaInSqMeters = UsdUtils::GetSpotLightArea(SourceRadiusUE / 100.0f);
 					const float FinalIntensityNits = UsdUtils::ConvertIntensityToNits(
 						SpotLightComponent->Intensity,
 						SolidAngle,
@@ -5175,13 +5176,13 @@ bool UnrealToUsd::CreateComponentPropertyBaker(
 
 					IntensitySetter(FinalIntensityNits, UsdTimeCode);
 					TreatAsPointSetter(FMath::IsNearlyZero(SpotLightComponent->SourceRadius), UsdTimeCode);
-					RadiusSetter(SpotLightComponent->SourceRadius, UsdTimeCode);
-					AngleSetter(SpotLightComponent->OuterConeAngle, UsdTimeCode);
-
-					const float Softness = FMath::IsNearlyZero(SpotLightComponent->OuterConeAngle)
-											   ? 0.0
-											   : 1.0f - SpotLightComponent->InnerConeAngle / SpotLightComponent->OuterConeAngle;
-					SoftnessSetter(Softness, UsdTimeCode);
+					RadiusSetter(SourceRadiusUE, UsdTimeCode);
+					
+					float UsdConeAngle = 90.0;
+					float UsdConeSoftness = 0.0f;
+					UnrealToUsd::ConvertSpotLightConeAngles(SpotLightComponent->OuterConeAngle, SpotLightComponent->InnerConeAngle, UsdConeAngle, UsdConeSoftness);
+					AngleSetter(UsdConeAngle, UsdTimeCode);
+					SoftnessSetter(UsdConeSoftness, UsdTimeCode);
 
 					ColorSetter(SpotLightComponent->LightColor, UsdTimeCode);
 					EnableTemperatureSetter(SpotLightComponent->bUseTemperature, UsdTimeCode);
@@ -5220,7 +5221,8 @@ bool UnrealToUsd::CreateComponentPropertyBaker(
 								 ColorSetter,
 								 EnableTemperatureSetter,
 								 TemperatureSetter,
-								 ShadowSetter](double UsdTimeCode) mutable
+								 ShadowSetter,
+								 MinLightSourceSize](double UsdTimeCode) mutable
 				{
 					TRACE_CPUPROFILER_EVENT_SCOPE(PointLightBaker);
 
@@ -5230,11 +5232,9 @@ bool UnrealToUsd::CreateComponentPropertyBaker(
 						return;
 					}
 
-					const float SolidAngle = 4.f * PI;
-					const float AreaInSqMeters = FMath::Max(
-						SolidAngle * FMath::Square(PointLightComponent->SourceRadius / 100.f),
-						KINDA_SMALL_NUMBER
-					);
+					const float SourceRadiusUE = FMath::Max(PointLightComponent->SourceRadius, MinLightSourceSize);
+					const float SolidAngle = UsdUtils::GetPointLightSolidAngle();
+					const float AreaInSqMeters = UsdUtils::GetPointLightArea(SourceRadiusUE / 100.0f);
 					const float FinalIntensityNits = UsdUtils::ConvertIntensityToNits(
 						PointLightComponent->Intensity,
 						SolidAngle,
@@ -5244,7 +5244,7 @@ bool UnrealToUsd::CreateComponentPropertyBaker(
 
 					IntensitySetter(FinalIntensityNits, UsdTimeCode);
 					TreatAsPointSetter(FMath::IsNearlyZero(PointLightComponent->SourceRadius), UsdTimeCode);
-					RadiusSetter(PointLightComponent->SourceRadius, UsdTimeCode);
+					RadiusSetter(SourceRadiusUE, UsdTimeCode);
 
 					ColorSetter(PointLightComponent->LightColor, UsdTimeCode);
 					EnableTemperatureSetter(PointLightComponent->bUseTemperature, UsdTimeCode);
