[USD] Adds FSdfPath::ReplacePrefix() and FSdfPath::GetPathElementCount() wrappers.

#Synced-CL 52444545

diff --git a/Engine/Plugins/Runtime/USDCore/Source/UnrealUSDWrapper/Private/UsdWrappers/SdfPath.cpp b/Engine/Plugins/Runtime/USDCore/Source/UnrealUSDWrapper/Private/UsdWrappers/SdfPath.cpp
--- a/Engine/Plugins/Runtime/USDCore/Source/UnrealUSDWrapper/Private/UsdWrappers/SdfPath.cpp
+++ b/Engine/Plugins/Runtime/USDCore/Source/UnrealUSDWrapper/Private/UsdWrappers/SdfPath.cpp
@@ -310,6 +310,17 @@ namespace UE
 #endif	  // #if USE_USD_SDK
 	}
 
+	FSdfPath FSdfPath::ReplacePrefix(const UE::FSdfPath& OldPrefix,
+								     const UE::FSdfPath& NewPrefix,
+			                         bool FixTargetPaths) const
+	{
+#if USE_USD_SDK
+		return FSdfPath(Impl->PxrSdfPath.Get().ReplacePrefix(OldPrefix, NewPrefix, FixTargetPaths));
+#else
+		return FSdfPath();
+#endif	  // #if USE_USD_SDK
+	}
+
 	FSdfPath FSdfPath::GetParentPath() const
 	{
 #if USE_USD_SDK
@@ -390,6 +401,17 @@ namespace UE
 		return Result;
 	}
 
+	int32 FSdfPath::GetPathElementCount() const
+	{
+#if USE_USD_SDK
+		FScopedUsdAllocs Allocs;
+
+		return static_cast<int32>(Impl->PxrSdfPath.Get().GetPathElementCount());
+
+#endif	  // #if USE_USD_SDK
+		return -1;
+	}
+
 	bool FSdfPath::HasPrefix(const UE::FSdfPath& Prefix) const
 	{
 #if USE_USD_SDK
diff --git a/Engine/Plugins/Runtime/USDCore/Source/UnrealUSDWrapper/Public/UsdWrappers/SdfPath.h b/Engine/Plugins/Runtime/USDCore/Source/UnrealUSDWrapper/Public/UsdWrappers/SdfPath.h
--- a/Engine/Plugins/Runtime/USDCore/Source/UnrealUSDWrapper/Public/UsdWrappers/SdfPath.h
+++ b/Engine/Plugins/Runtime/USDCore/Source/UnrealUSDWrapper/Public/UsdWrappers/SdfPath.h
@@ -78,6 +78,9 @@ namespace UE
 		UE_API FSdfPath GetAbsoluteRootOrPrimPath() const;
 
 		UE_API FSdfPath ReplaceName(const TCHAR* NewLeafName) const;
+		UE_API FSdfPath ReplacePrefix(const UE::FSdfPath& OldPrefix, 
+			                          const UE::FSdfPath& NewPrefix, 
+			                          bool FixTargetPaths = true) const;
 
 		UE_API FSdfPath GetParentPath() const;
 		UE_API FSdfPath GetPrimPath() const;
@@ -90,6 +93,7 @@ namespace UE
 		UE_API FString GetString() const;
 
 		UE_API TArray<FSdfPath> GetPrefixes() const;
+		UE_API int32 GetPathElementCount() const;
 
 		UE_API bool HasPrefix(const UE::FSdfPath& Prefix) const;
 
