[USD] Fix if/endif mistake from previous checkin

#Synced-CL 52444545

diff --git a/Engine/Plugins/Runtime/USDCore/Source/UnrealUSDWrapper/Private/UsdWrappers/SdfPath.cpp b/Engine/Plugins/Runtime/USDCore/Source/UnrealUSDWrapper/Private/UsdWrappers/SdfPath.cpp
--- a/Engine/Plugins/Runtime/USDCore/Source/UnrealUSDWrapper/Private/UsdWrappers/SdfPath.cpp
+++ b/Engine/Plugins/Runtime/USDCore/Source/UnrealUSDWrapper/Private/UsdWrappers/SdfPath.cpp
@@ -407,9 +407,9 @@ namespace UE
 		FScopedUsdAllocs Allocs;
 
 		return static_cast<int32>(Impl->PxrSdfPath.Get().GetPathElementCount());
-
-#endif	  // #if USE_USD_SDK
+#else
 		return -1;
+#endif	  // #if USE_USD_SDK
 	}
 
 	bool FSdfPath::HasPrefix(const UE::FSdfPath& Prefix) const
